#UE4
- Updated MemoryProfiler2 to support the latest platform memory changes, should be able to load .mprof files from version 3 and 4.

diff --git a/Engine/Source/Programs/MemoryProfiler2/MainWindow.cs b/Engine/Source/Programs/MemoryProfiler2/MainWindow.cs
--- a/Engine/Source/Programs/MemoryProfiler2/MainWindow.cs
+++ b/Engine/Source/Programs/MemoryProfiler2/MainWindow.cs
@@ -1853,7 +1853,7 @@ namespace MemoryProfiler2
 			}
 		}
 
-		/// <summary> Class used to sort columns by specified column in the exlusive list view. </summary>
+		/// <summary> Class used to sort columns by specified column in the exclusive list view. </summary>
 		public class FColumnSorter : IComparer
 		{
 			public int ColumnToSortBy = 0;
diff --git a/Engine/Source/Programs/MemoryProfiler2/ProfileDataHeader.cs b/Engine/Source/Programs/MemoryProfiler2/ProfileDataHeader.cs
--- a/Engine/Source/Programs/MemoryProfiler2/ProfileDataHeader.cs
+++ b/Engine/Source/Programs/MemoryProfiler2/ProfileDataHeader.cs
@@ -20,6 +20,8 @@ namespace MemoryProfiler2
 		public UInt32 Version;
 		/** Platform that was captured.								*/
         public EPlatformType Platform;
+		/** Platform that was captured.								*/
+		public string PlatformName;
 		/** Whether symbol information was serialized.				*/
 		public bool	bShouldSerializeSymbolInfo;
 		/** Name of executable this information was gathered with.	*/
@@ -73,8 +75,19 @@ namespace MemoryProfiler2
 			{
 				// Version info for backward compatible serialization.
 				Version = BinaryStream.ReadUInt32();
-				// Platform and max backtrace depth.
-				Platform = (EPlatformType)BinaryStream.ReadUInt32();
+				
+				if( Version >= 4 )
+				{
+					// Read platform name.
+					PlatformName = FStreamParser.ReadString( BinaryStream );
+					PlatformName = FixFixedSizeString( PlatformName );
+				}
+				else
+				{
+					// Read platform type.
+					Platform = (EPlatformType)BinaryStream.ReadUInt32();
+				}
+
 				// Whether symbol information was serialized.
 				bShouldSerializeSymbolInfo = BinaryStream.ReadUInt32() == 0 ? false : true;
 
@@ -97,7 +110,11 @@ namespace MemoryProfiler2
 				NumDataFiles = BinaryStream.ReadUInt32();
 
                 FStreamToken.Version = Version;
-				if (Version > 2)
+				if( Version >= 4 )
+				{
+					// Just ignore
+				}
+				else if (Version > 2)
 				{
 	                ScriptCallstackTableOffset = BinaryStream.ReadUInt32();
                     ScriptNameTableOffset = BinaryStream.ReadUInt32();
@@ -106,15 +123,20 @@ namespace MemoryProfiler2
                 }
 
 				// Name of executable.
-				UInt32 ExecutableNameLength = BinaryStream.ReadUInt32();
-				ExecutableName = new string(BinaryStream.ReadChars((int)ExecutableNameLength));
-				// We serialize a fixed size string. Trim the null characters that make it in by converting char[] to string.
-				int RealLength = 0;
-				while( ExecutableName[RealLength++] != '\0' ) 
-				{
-				}
-				ExecutableName = ExecutableName.Remove(RealLength-1);
+				ExecutableName = FStreamParser.ReadString( BinaryStream );
+				ExecutableName = FixFixedSizeString( ExecutableName );			
+			}
+		}
+
+		/// <summary> Trims the null characters that make it in by converting char[] to string. </summary>
+		static String FixFixedSizeString( string String )
+		{
+			int RealLength = 0;
+			while( String[RealLength++] != '\0' )
+			{
 			}
+			string Result = String.Remove( RealLength - 1 );
+			return Result;
 		}
 	}
 
diff --git a/Engine/Source/Programs/MemoryProfiler2/StreamInfo.cs b/Engine/Source/Programs/MemoryProfiler2/StreamInfo.cs
--- a/Engine/Source/Programs/MemoryProfiler2/StreamInfo.cs
+++ b/Engine/Source/Programs/MemoryProfiler2/StreamInfo.cs
@@ -49,9 +49,12 @@ namespace MemoryProfiler2
 		/// <summary> List of snapshots created by parser and used by various visualizers and dumping tools. </summary>
 		public List<FStreamSnapshot> SnapshotList = new List<FStreamSnapshot>();
 
-		/// <summary> Platform on which these statistics were collected. </summary>
+		/// <summary> Platform on which these statistics were collected V3. </summary>
 		public EPlatformType Platform;
 
+		/// <summary> Platform that was captured V4. </summary>
+		public string PlatformName;
+
 		/// <summary> Profiling options, default all options are enabled. UI is not provided yet. </summary>
 		public ProfilingOptions CreationOptions;
 
@@ -91,12 +94,13 @@ namespace MemoryProfiler2
 		}
 
 		/// <summary> Initializes and sizes the arrays. Size is known as soon as header is serialized. </summary>
-		public void Initialize( EPlatformType InPlatform, uint InNameArrayEntries, uint InCallStackArrayEntries, uint InCallStackAddressArrayEntries )
+		public void Initialize( FProfileDataHeader Header )
 		{
-			Platform = InPlatform;
-			NameArray = new List<string>( ( int )InNameArrayEntries );
-			CallStackArray = new List<FCallStack>( ( int )InCallStackArrayEntries );
-			CallStackAddressArray = new List<FCallStackAddress>( ( int )InCallStackAddressArrayEntries );
+			Platform = Header.Platform;
+			PlatformName = Header.PlatformName;
+			NameArray = new List<string>( (int)Header.NameTableEntries );
+			CallStackArray = new List<FCallStack>( (int)Header.CallStackTableEntries );
+			CallStackAddressArray = new List<FCallStackAddress>( (int)Header.CallStackAddressTableEntries );
 		}
 
 		/// <summary> 
@@ -150,25 +154,32 @@ namespace MemoryProfiler2
 		{
 			Debug.Assert( BankIndex >= 0 );
 
-			//private static int[][] MEMORY_BANK_SIZE = new int[][] { new int[] { 512, 1 }, new int[] { 180, 256 } };
-			//private static int[][] MEMORY_BANK_SIZE_FOR_RENDERED_GRAPHS = new int[][] { new int[] { 512 * 1024 * 1024, 1 }, new int[] { 256 * 1024 * 1024, 256 * 1024 * 1024 } };
-			if (GlobalInstance.Platform == EPlatformType.Xbox360 || GlobalInstance.Platform == EPlatformType.IPhone)
+			if( FStreamToken.Version >= 4 )
 			{
 				return BankIndex == 0 ? 512 : 0;
 			}
-			else if( GlobalInstance.Platform == EPlatformType.PS3 )
-			{
-				return BankIndex < 2 ? 256 : 0;
-			}
-			else if( ( GlobalInstance.Platform & EPlatformType.AnyWindows ) != EPlatformType.Unknown )
-			{
-				// Technically could be much larger, but the graph becomes unusable
-				// if the scale is too large.
-				return BankIndex == 0 ? 512 : 0;
-			}
 			else
 			{
-				throw new Exception( "Unsupported platform" );
+				//private static int[][] MEMORY_BANK_SIZE = new int[][] { new int[] { 512, 1 }, new int[] { 180, 256 } };
+				//private static int[][] MEMORY_BANK_SIZE_FOR_RENDERED_GRAPHS = new int[][] { new int[] { 512 * 1024 * 1024, 1 }, new int[] { 256 * 1024 * 1024, 256 * 1024 * 1024 } };
+				if( GlobalInstance.Platform == EPlatformType.Xbox360 || GlobalInstance.Platform == EPlatformType.IPhone )
+				{
+					return BankIndex == 0 ? 512 : 0;
+				}
+				else if( GlobalInstance.Platform == EPlatformType.PS3 )
+				{
+					return BankIndex < 2 ? 256 : 0;
+				}
+				else if( ( GlobalInstance.Platform & EPlatformType.AnyWindows ) != EPlatformType.Unknown )
+				{
+					// Technically could be much larger, but the graph becomes unusable
+					// if the scale is too large.
+					return BankIndex == 0 ? 512 : 0;
+				}
+				else
+				{
+					throw new Exception( "Unsupported platform" );
+				}
 			}
 		}
 
diff --git a/Engine/Source/Programs/MemoryProfiler2/StreamParser.cs b/Engine/Source/Programs/MemoryProfiler2/StreamParser.cs
--- a/Engine/Source/Programs/MemoryProfiler2/StreamParser.cs
+++ b/Engine/Source/Programs/MemoryProfiler2/StreamParser.cs
@@ -84,7 +84,7 @@ namespace MemoryProfiler2
             int NextDataFile = 1;
 
 			// Initialize shared information across snapshots, namely names, callstacks and addresses.
-			FStreamInfo.GlobalInstance.Initialize(Header.Platform, Header.NameTableEntries, Header.CallStackTableEntries, Header.CallStackAddressTableEntries);
+			FStreamInfo.GlobalInstance.Initialize( Header );
 
 			// Keep track of current position as it's where the token stream starts.
 			long TokenStreamOffset = ParserFileStream.Position;
@@ -320,7 +320,8 @@ namespace MemoryProfiler2
 						MarkerSnapshot.MetricArray = new List<long>( Token.Metrics );
 						MarkerSnapshot.LoadedLevels = new List<int>( Token.LoadedLevels );
 						MarkerSnapshot.OverallMemorySlice = new List<FMemorySlice>( Snapshot.OverallMemorySlice );
-						MarkerSnapshot.MemoryAllocationStats = Token.MemoryAllocationStats.DeepCopy();
+						MarkerSnapshot.MemoryAllocationStats3 = Token.MemoryAllocationStats3.DeepCopy();
+						MarkerSnapshot.MemoryAllocationStats4 = Token.MemoryAllocationStats4.DeepCopy();
 
 						FStreamInfo.GlobalInstance.SnapshotList.Add( MarkerSnapshot );
 
@@ -543,7 +544,8 @@ namespace MemoryProfiler2
 										MarkerSnapshot.SnapshotIndex = SnapshotIndex;
 										MarkerSnapshot.MetricArray = new List<long>( Token.Metrics );
 										MarkerSnapshot.LoadedLevels = new List<int>( Token.LoadedLevels );
-										MarkerSnapshot.MemoryAllocationStats = Snapshot.MemoryAllocationStats.DeepCopy();
+										MarkerSnapshot.MemoryAllocationStats3 = Snapshot.MemoryAllocationStats3.DeepCopy();
+										MarkerSnapshot.MemoryAllocationStats4 = Snapshot.MemoryAllocationStats4.DeepCopy();
 										MarkerSnapshot.OverallMemorySlice = new List<FMemorySlice>( Snapshot.OverallMemorySlice );
 
 										FStreamInfo.GlobalInstance.SnapshotList.Add( MarkerSnapshot );
@@ -558,61 +560,61 @@ namespace MemoryProfiler2
 
 								case EProfilingPayloadSubType.SUBTYPE_TotalUsed:
 								{
-									Snapshot.MemoryAllocationStats.TotalUsed = Token.Payload;
+									Snapshot.MemoryAllocationStats3.TotalUsed = Token.Payload;
 									break;
 								}
 
 								case EProfilingPayloadSubType.SUBTYPE_TotalAllocated:
 								{
-									Snapshot.MemoryAllocationStats.TotalAllocated = Token.Payload;
+									Snapshot.MemoryAllocationStats3.TotalAllocated = Token.Payload;
 									break;
 								}
 
 								case EProfilingPayloadSubType.SUBTYPE_CPUUsed:
 								{
-									Snapshot.MemoryAllocationStats.CPUUsed = Token.Payload;
+									Snapshot.MemoryAllocationStats3.CPUUsed = Token.Payload;
 									break;
 								}
 
 								case EProfilingPayloadSubType.SUBTYPE_CPUSlack:
 								{
-									Snapshot.MemoryAllocationStats.CPUSlack = Token.Payload;
+									Snapshot.MemoryAllocationStats3.CPUSlack = Token.Payload;
 									break;
 								}
 
 								case EProfilingPayloadSubType.SUBTYPE_CPUWaste:
 								{
-									Snapshot.MemoryAllocationStats.CPUWaste = Token.Payload;
+									Snapshot.MemoryAllocationStats3.CPUWaste = Token.Payload;
 									break;
 								}
 
 								case EProfilingPayloadSubType.SUBTYPE_GPUUsed:
 								{
-									Snapshot.MemoryAllocationStats.GPUUsed = Token.Payload;
+									Snapshot.MemoryAllocationStats3.GPUUsed = Token.Payload;
 									break;
 								}
 
 								case EProfilingPayloadSubType.SUBTYPE_GPUSlack:
 								{
-									Snapshot.MemoryAllocationStats.GPUSlack = Token.Payload;
+									Snapshot.MemoryAllocationStats3.GPUSlack = Token.Payload;
 									break;
 								}
 
 								case EProfilingPayloadSubType.SUBTYPE_GPUWaste:
 								{
-									Snapshot.MemoryAllocationStats.GPUWaste = Token.Payload;
+									Snapshot.MemoryAllocationStats3.GPUWaste = Token.Payload;
 									break;
 								}
 
 								case EProfilingPayloadSubType.SUBTYPE_ImageSize:
 								{
-									Snapshot.MemoryAllocationStats.ImageSize = Token.Payload;
+									Snapshot.MemoryAllocationStats3.ImageSize = Token.Payload;
 									break;
 								}
 
 								case EProfilingPayloadSubType.SUBTYPE_OSOverhead:
 								{
-									Snapshot.MemoryAllocationStats.OSOverhead = Token.Payload;
+									Snapshot.MemoryAllocationStats3.OSOverhead = Token.Payload;
 									break;
 								}
 
@@ -630,7 +632,8 @@ namespace MemoryProfiler2
 
 								case EProfilingPayloadSubType.SUBTYPE_MemoryAllocationStats:
 								{
-									Snapshot.MemoryAllocationStats = Token.MemoryAllocationStats.DeepCopy();
+									Snapshot.MemoryAllocationStats3 = Token.MemoryAllocationStats3.DeepCopy();
+									Snapshot.MemoryAllocationStats4 = Token.MemoryAllocationStats4.DeepCopy();
 									break;
 								}
 
@@ -746,7 +749,8 @@ namespace MemoryProfiler2
 			Snapshot.ElapsedTime = Token.ElapsedTime;
             Snapshot.MetricArray = new List<long>(Token.Metrics);
             Snapshot.LoadedLevels = new List<int>(Token.LoadedLevels);
-			Snapshot.MemoryAllocationStats = Token.MemoryAllocationStats.DeepCopy();
+			Snapshot.MemoryAllocationStats3 = Token.MemoryAllocationStats3.DeepCopy();
+			Snapshot.MemoryAllocationStats4 = Token.MemoryAllocationStats4.DeepCopy();
             Snapshot.FinalizeSnapshot(PointerToPointerInfoMap);
             FStreamInfo.GlobalInstance.SnapshotList.Add(Snapshot);
 
@@ -1069,47 +1073,47 @@ namespace MemoryProfiler2
 			}
 			else
 			{
-			foreach( FCallStackAddress Address in FStreamInfo.GlobalInstance.CallStackAddressArray )
-			{
-                if ((AddressIndex % NumAddressesPerTick) == 0)
-                {
-                    BGWorker.ReportProgress(AddressIndex / NumAddressesPerTick, ProgressString);
-                }
-                ++AddressIndex;
-
-                // Look up symbol info via console support DLL.
-				string Filename = "";
-				string Function = "";
-                FStreamInfo.GlobalInstance.ConsoleSymbolParser.ResolveAddressToSymboInfo((uint)Address.ProgramCounter, ref Filename, ref Function, ref Address.LineNumber);
-	
-				// Look up filename index.
-				if( NameToIndexMap.ContainsKey(Filename) )
+				foreach( FCallStackAddress Address in FStreamInfo.GlobalInstance.CallStackAddressArray )
 				{
-					// Use existing entry.
-					Address.FilenameIndex = NameToIndexMap[Filename];
-				}
-				// Not found, so we use global name index to set new one.
-				else
-				{
-					// Set name in map associated with global ever increasing index.
-					Address.FilenameIndex = CurrentNameIndex++;
-					NameToIndexMap.Add( Filename, Address.FilenameIndex );
-				}
+					if( ( AddressIndex % NumAddressesPerTick ) == 0 )
+					{
+						BGWorker.ReportProgress( AddressIndex / NumAddressesPerTick, ProgressString );
+					}
+					++AddressIndex;
 
-				// Look up function index.
-				if( NameToIndexMap.ContainsKey(Function) )
-				{
-					// Use existing entry.
-					Address.FunctionIndex = NameToIndexMap[Function];
+					// Look up symbol info via console support DLL.
+					string Filename = "";
+					string Function = "";
+					FStreamInfo.GlobalInstance.ConsoleSymbolParser.ResolveAddressToSymboInfo( (uint)Address.ProgramCounter, ref Filename, ref Function, ref Address.LineNumber );
+
+					// Look up filename index.
+					if( NameToIndexMap.ContainsKey( Filename ) )
+					{
+						// Use existing entry.
+						Address.FilenameIndex = NameToIndexMap[Filename];
+					}
+					// Not found, so we use global name index to set new one.
+					else
+					{
+						// Set name in map associated with global ever increasing index.
+						Address.FilenameIndex = CurrentNameIndex++;
+						NameToIndexMap.Add( Filename, Address.FilenameIndex );
+					}
+
+					// Look up function index.
+					if( NameToIndexMap.ContainsKey( Function ) )
+					{
+						// Use existing entry.
+						Address.FunctionIndex = NameToIndexMap[Function];
+					}
+					// Not found, so we use global name index to set new one.
+					else
+					{
+						// Set name in map associated with global ever increasing index.
+						Address.FunctionIndex = CurrentNameIndex++;
+						NameToIndexMap.Add( Function, Address.FunctionIndex );
+					}
 				}
-				// Not found, so we use global name index to set new one.
-				else
-				{
-					// Set name in map associated with global ever increasing index.
-					Address.FunctionIndex = CurrentNameIndex++;
-					NameToIndexMap.Add( Function, Address.FunctionIndex );
-				} 
-			}
 			}
 
 			// Create new name array based on dictionary.
diff --git a/Engine/Source/Programs/MemoryProfiler2/StreamSnapshot.cs b/Engine/Source/Programs/MemoryProfiler2/StreamSnapshot.cs
--- a/Engine/Source/Programs/MemoryProfiler2/StreamSnapshot.cs
+++ b/Engine/Source/Programs/MemoryProfiler2/StreamSnapshot.cs
@@ -85,7 +85,7 @@ namespace MemoryProfiler2
     }
 
 	/// <summary> State of total memory at a specific moment in time. </summary>
-	public enum ESliceTypes
+	public enum ESliceTypesV3
 	{
 		TotalUsed,
 		TotalAllocated,
@@ -104,34 +104,65 @@ namespace MemoryProfiler2
 		Count
 	};
 
+	public enum ESliceTypesV4
+	{
+		PlatformUsedPhysical,
+		BinnedWasteCurrent,
+		BinnedUsedCurrent,
+		BinnedSlackCurrent,
+		MemoryProfilingOverhead,
+		OverallAllocatedMemory,
+		Count
+	};
+
 	public class FMemorySlice
 	{
 		public long[] MemoryInfo = null;
 
 		public FMemorySlice( FStreamSnapshot Snapshot )
 		{
-			MemoryInfo = new long[]
+			if( FStreamToken.Version >= 4 )
 			{
-				Snapshot.MemoryAllocationStats.TotalUsed,
-				Snapshot.MemoryAllocationStats.TotalAllocated,
-				Snapshot.MemoryAllocationStats.CPUUsed,
-				Snapshot.MemoryAllocationStats.CPUSlack,
-				Snapshot.MemoryAllocationStats.CPUWaste,
-				Snapshot.MemoryAllocationStats.GPUUsed,
-				Snapshot.MemoryAllocationStats.GPUSlack,
-				Snapshot.MemoryAllocationStats.GPUWaste,
-				Snapshot.MemoryAllocationStats.OSOverhead,
-				Snapshot.MemoryAllocationStats.ImageSize,
-				Snapshot.MemoryAllocationStats.HostUsed,
-				Snapshot.MemoryAllocationStats.HostSlack,
-				Snapshot.MemoryAllocationStats.HostWaste,
-				Snapshot.AllocationSize,
-			};
+				MemoryInfo = new long[]
+				{
+					Snapshot.MemoryAllocationStats4[FMemoryAllocationStatsV4.PlatformUsedPhysical],
+					Snapshot.MemoryAllocationStats4[FMemoryAllocationStatsV4.BinnedWasteCurrent],
+					Snapshot.MemoryAllocationStats4[FMemoryAllocationStatsV4.BinnedUsedCurrent],
+					Snapshot.MemoryAllocationStats4[FMemoryAllocationStatsV4.BinnedSlackCurrent],
+					Snapshot.MemoryAllocationStats4[FMemoryAllocationStatsV4.MemoryProfilingOverhead],
+					Snapshot.AllocationSize,
+				};
+			}
+			else
+			{
+				MemoryInfo = new long[]
+				{
+					Snapshot.MemoryAllocationStats3.TotalUsed,
+					Snapshot.MemoryAllocationStats3.TotalAllocated,
+					Snapshot.MemoryAllocationStats3.CPUUsed,
+					Snapshot.MemoryAllocationStats3.CPUSlack,
+					Snapshot.MemoryAllocationStats3.CPUWaste,
+					Snapshot.MemoryAllocationStats3.GPUUsed,
+					Snapshot.MemoryAllocationStats3.GPUSlack,
+					Snapshot.MemoryAllocationStats3.GPUWaste,
+					Snapshot.MemoryAllocationStats3.OSOverhead,
+					Snapshot.MemoryAllocationStats3.ImageSize,
+					Snapshot.MemoryAllocationStats3.HostUsed,
+					Snapshot.MemoryAllocationStats3.HostSlack,
+					Snapshot.MemoryAllocationStats3.HostWaste,
+					Snapshot.AllocationSize,
+				};
+			}
 		}
 
-		public long GetSliceInfo( ESliceTypes SliceType )
+		public long GetSliceInfoV3( ESliceTypesV3 SliceType )
 		{
-			return ( MemoryInfo[( int )SliceType] );
+			return ( MemoryInfo[(int)SliceType] );
+		}
+
+		public long GetSliceInfoV4( ESliceTypesV4 SliceType )
+		{
+			return ( MemoryInfo[(int)SliceType] );
 		}
 	}
 
@@ -175,7 +206,10 @@ namespace MemoryProfiler2
 		public List<string> LoadedLevelNames = new List<string>();
 
 		/// <summary> Generic memory allocation stats. </summary>
-		public FMemoryAllocationStats MemoryAllocationStats = new FMemoryAllocationStats();
+		public FMemoryAllocationStatsV3 MemoryAllocationStats3 = new FMemoryAllocationStatsV3();
+
+		/// <summary> Generic memory allocation stats. </summary>
+		public FMemoryAllocationStatsV4 MemoryAllocationStats4 = new FMemoryAllocationStatsV4();
 
 		/// <summary> Running count of number of allocations. </summary>
 		public long AllocationCount = 0;
@@ -209,7 +243,7 @@ namespace MemoryProfiler2
 		/// <summary> Performs a deep copy of the relevant data structures. </summary>
 		public FStreamSnapshot DeepCopy( Dictionary<ulong, FCallStackAllocationInfo> PointerToPointerInfoMap )
 		{
-			// Create new snaphot object.
+			// Create new snapshot object.
 			FStreamSnapshot Snapshot = new FStreamSnapshot( "Copy" );
 
 			// Manually perform a deep copy of LifetimeCallstackList
@@ -287,7 +321,9 @@ namespace MemoryProfiler2
 					ResultSnapshot.MetricArray.Add( New.MetricArray[ CallstackIndex ] - Old.MetricArray[ CallstackIndex ] );
 				}
 
-				ResultSnapshot.MemoryAllocationStats = FMemoryAllocationStats.Diff( Old.MemoryAllocationStats, New.MemoryAllocationStats );
+				
+				ResultSnapshot.MemoryAllocationStats3 = FMemoryAllocationStatsV3.Diff( Old.MemoryAllocationStats3, New.MemoryAllocationStats3 );
+				ResultSnapshot.MemoryAllocationStats4 = FMemoryAllocationStatsV4.Diff( Old.MemoryAllocationStats4, New.MemoryAllocationStats4 );
 
 				ResultSnapshot.StreamIndex = New.StreamIndex;
 				ResultSnapshot.bIsDiffResult = true;
@@ -496,7 +532,11 @@ namespace MemoryProfiler2
 			StrBuilder.AppendLine( "Metrics: " );
 			if( MetricArray.Count > 0 )
 			{
-				if( FStreamToken.Version > 2 && MetricArray.Count == (int)ESnapshotMetricV3.Count )
+				if( FStreamToken.Version >= 4 )
+				{
+					// Just ignore
+				}
+				else if( FStreamToken.Version > 2 && MetricArray.Count == (int)ESnapshotMetricV3.Count )
 				{
 					string[] MetricNames = Enum.GetNames( typeof( ESnapshotMetricV3 ) );
 					for( int MetricIndex = 0; MetricIndex < MetricNames.Length - 1; MetricIndex++ )
@@ -523,9 +563,18 @@ namespace MemoryProfiler2
 
 			///// <summary> Generic memory allocation stats. </summary>
 			//public FMemoryAllocationStats MemoryAllocationStats = new FMemoryAllocationStats();
-			StrBuilder.AppendLine( "Memory Allocation Stats: " );
-			StrBuilder.Append( MemoryAllocationStats.ToString() );
 
+			if( FStreamToken.Version >= 4 )
+			{
+				StrBuilder.AppendLine( "Memory Allocation Stats: " );
+				StrBuilder.Append( MemoryAllocationStats4.ToString() );
+			}
+			else
+			{
+				StrBuilder.AppendLine( "Memory Allocation Stats: " );
+				StrBuilder.Append( MemoryAllocationStats3.ToString() );
+			}
+			
 			///// <summary> Running count of number of allocations. </summary>
 			//public long AllocationCount = 0;
 			StrBuilder.AppendLine( "Allocation Count: " + AllocationCount );
diff --git a/Engine/Source/Programs/MemoryProfiler2/StreamToken.cs b/Engine/Source/Programs/MemoryProfiler2/StreamToken.cs
--- a/Engine/Source/Programs/MemoryProfiler2/StreamToken.cs
+++ b/Engine/Source/Programs/MemoryProfiler2/StreamToken.cs
@@ -11,14 +11,14 @@ using System.Diagnostics;
 namespace MemoryProfiler2
 {
 	/// <summary> The lower 2 bits of a pointer are piggy-bagged to store what kind of data follows it. This enum lists the possible types. </summary>
-    public enum EProfilingPayloadType
-    {
-        TYPE_Malloc		= 0,
-        TYPE_Free		= 1,
-        TYPE_Realloc	= 2,
-        TYPE_Other		= 3,
-        // Don't add more than 4 values - we only have 2 bits to store this.
-    }
+	public enum EProfilingPayloadType
+	{
+		TYPE_Malloc		= 0,
+		TYPE_Free		= 1,
+		TYPE_Realloc	= 2,
+		TYPE_Other		= 3,
+		// Don't add more than 4 values - we only have 2 bits to store this.
+	}
 
 	/// <summary> 
 	///	The the case of TYPE_Other, this enum determines the subtype of the token.
@@ -99,7 +99,7 @@ namespace MemoryProfiler2
 		/// <summary> Marker used to determine when a previously streamed level has been made visible. </summary> 
 		SUBTYPE_SnapshotMarker_LevelStream_End = 27,
 
-		/// <summary> Marker used to store a generic malloc statistics. @see FMemoryAllocationStats. </summary> 
+		/// <summary> Marker used to store a generic malloc statistics. @see FMallocProfiler::WriteMemoryAllocationStats </summary> 
 		SUBTYPE_MemoryAllocationStats = 31,
 
 		/// <summary> Start licensee-specific subtypes from here. </summary> 
@@ -109,65 +109,13 @@ namespace MemoryProfiler2
 		SUBTYPE_Unknown,
 	};
 
-	/// <summary> Enumeration of all memory allocation statistics types. </summary>
-	public enum FMemoryAllocationTypes
+	public interface IMemoryStats
 	{
-		/// <summary> The total amount of memory used by the game. </summary>
-		TotalUsed,
-		/// <summary> The total amount of memory allocated from the OS. </summary>
-		TotalAllocated,
-
-		// Virtual memory for Xbox and PC / Main memory for PS3 (tracked in the allocators) Host is not included ??
-
-		/// <summary> The allocated in use by the application virtual memory. </summary>
-		CPUUsed,
-		/// <summary> The allocated from the OS/allocator, but not used by the application. </summary>
-		CPUSlack,
-		/// <summary> Alignment waste from a pooled allocator plus book keeping overhead. </summary>
-		CPUWaste,
-		/// <summary> The amount of free memory before the first malloc has been done. (PS3 only). </summary>
-		CPUAvailable,
-
-		// Physical memory for Xbox and PC / Local memory for PS3 (tracked in the allocators)
-
-		/// <summary> The allocated in use by the application physical memory. </summary>
-		GPUUsed,
-		/// <summary> The allocated from the OS, but not used by the application. </summary>
-		GPUSlack,
-		/// <summary> Alignment waste from a pooled allocator plus book keeping overhead. </summary>
-		GPUWaste,
-		/// <summary> The total amount of memory available for the allocator. (PS3 only) </summary>
-		GPUAvailable,
-
-		/// <summary> Used memory as reported by the operating system. (Xbox360, PS3 only) </summary>
-		OSReportedUsed,
-		/// <summary> Free memory as reported by the operating system. (Xbox360, PS3 only) </summary>
-		OSReportedFree,
-		/// <summary> The overhead of the operating system. (Xbox360 only) </summary>
-		OSOverhead,
-		/// <summary> Size of loaded executable, stack, static, and global object size. (Xbox360, PS3 only) </summary>
-		ImageSize,
-
-		/// <summary> Host memory in use by the application. (PS3 only) </summary>
-		HostUsed,
-		/// <summary> Host memory allocated, but not used by the application. (PS3 only) </summary>
-		HostSlack,
-		/// <summary> Host memory wasted due to allocations' alignment. (PS3 only) </summary>
-		HostWaste,
-		/// <summary> The total amount of host memory that has been allocated. (PS3 only) </summary>
-		HostAvailable,
-
-		/// <summary> Size of allocated memory in the texture pool. </summary>
-		AllocatedTextureMemorySize,
-		/// <summary> Size of available memory in the texture pool. </summary>
-		AvailableTextureMemorySize,
 
-		/// <summary> Last item used as a count. </summary>
-		Count,
 	}
 
 	/// <summary> Struct used to hold memory allocation statistics. Mirrored in MemoryBase.h. This need to be a class because of how ReadMemoryAllocationsStats method works. </summary>
-	public class FMemoryAllocationStats
+	public class FMemoryAllocationStatsV3
 	{
 		/// <summary> The total amount of memory used by the game. </summary>
 		public Int64 TotalUsed;
@@ -249,11 +197,11 @@ namespace MemoryProfiler2
 		}
 
 		/// <summary> Returns a difference between old and new memory allocation stats. </summary>
-		public static FMemoryAllocationStats Diff( FMemoryAllocationStats Old, FMemoryAllocationStats New )
+		public static FMemoryAllocationStatsV3 Diff( FMemoryAllocationStatsV3 Old, FMemoryAllocationStatsV3 New )
 		{
-			FMemoryAllocationStats Diff = new FMemoryAllocationStats();
+			FMemoryAllocationStatsV3 Diff = new FMemoryAllocationStatsV3();
 
-			FieldInfo[] FieldInfos = typeof( FMemoryAllocationStats ).GetFields();
+			FieldInfo[] FieldInfos = typeof( FMemoryAllocationStatsV3 ).GetFields();
 			int PropertiesNum = FieldInfos.Length;
 
 			for( int FieldIndex = 0; FieldIndex < PropertiesNum; FieldIndex++ )
@@ -268,9 +216,9 @@ namespace MemoryProfiler2
 		}
 
 		/// <summary> Creates a new copy of this class. </summary>
-		public FMemoryAllocationStats DeepCopy()
+		public FMemoryAllocationStatsV3 DeepCopy()
 		{
-			FMemoryAllocationStats Copy = ( FMemoryAllocationStats )MemberwiseClone();
+			FMemoryAllocationStatsV3 Copy = ( FMemoryAllocationStatsV3 )MemberwiseClone();
 			return Copy;
 		}
 
@@ -353,58 +301,176 @@ namespace MemoryProfiler2
 		}
 	};
 
+	/// <summary> 
+	/// Struct used to hold platform memory stats and allocator stats. Valid only for MemoryProfiler2.FProfileDataHeader.Version >= 4.
+	/// @see FGenericPlatformMemory::GetStatsForMallocProfiler 
+	/// @see FMalloc::GetAllocatorStats
+	/// @todo add support for non-generic platform memory stats.
+	/// </summary>
+	public class FMemoryAllocationStatsV4
+	{
+		// Names of generic platform memory stats. @see FGenericPlatformMemoryStats
+		
+		/// <summary> The amount of actual physical memory, in bytes. </summary>
+		static public string PlatformTotalPhysical = "Total Physical";
+
+		/// <summary> The amount of virtual memory, in bytes. </summary>
+		static public string PlatformTotalVirtual = "Total Virtual";
+
+		/// <summary> The size of a page, in bytes. </summary>
+		static public string PlatformPageSize = "Page Size";
+
+		/// <summary> Approximate physical RAM in GB"; 1 on everything except PC. Used for "course tuning", like FPlatformMisc::NumberOfCores(). </summary>
+		static public string PlatformTotalPhysicalGB = "Total Physical GB";
+
+		/// <summary> The amount of physical memory currently available, in bytes. </summary>
+		static public string PlatformAvailablePhysical = "Available Physical";
+
+		/// <summary> The amount of virtual memory currently available, in bytes. </summary>
+		static public string PlatformAvailableVirtual = "Available Virtual";
+
+		/// <summary> The amount of physical memory used by the process, in bytes. </summary>
+		static public string PlatformUsedPhysical = "Used Physical";
+
+		/// <summary> The peak amount of physical memory used by the process, in bytes. </summary>
+		static public string PlatformPeakUsedPhysical = "Peak Used Physical";
+
+		/// <summary> Total amount of virtual memory used by the process. </summary>
+		static public string PlatformUsedVirtual = "Used Virtual";
+
+		/// <summary> The peak amount of virtual memory used by the process. </summary>
+		static public string PlatformPeakUsedVirtual = "Peak Used Virtual";
+
+		static public string MemoryProfilingOverhead = "Memory Profiling Overhead";
+
+		// Names of malloc binned memory stats. @see FMallocBinned.GetAllocatorStats 
+
+		static public string BinnedWasteCurrent = "Binned Waste Current";
+		static public string BinnedUsedCurrent = "Binned Used Current";
+		static public string BinnedSlackCurrent = "Binned Slack Current";
+
+		private Dictionary<string, Int64> _Stats;
+
+		public FMemoryAllocationStatsV4()
+		{
+			_Stats = new Dictionary<string, Int64>();
+		}
+
+		public FMemoryAllocationStatsV4( FMemoryAllocationStatsV4 Other )
+		{
+			this._Stats = new Dictionary<string, Int64>( Other._Stats );
+		}
+
+		public Int64 this[string StatName]
+		{
+			get
+			{
+				return _Stats[StatName];
+			}
+			set
+			{
+				_Stats[StatName] = value;
+			}
+		}
+
+		public Int64 NumStats()
+		{
+			return _Stats.Count;
+		}
+
+		/// <summary> Returns a difference between old and new platform memory stats. </summary>
+		public static FMemoryAllocationStatsV4 Diff( FMemoryAllocationStatsV4 Old, FMemoryAllocationStatsV4 New )
+		{
+			FMemoryAllocationStatsV4 Diff = new FMemoryAllocationStatsV4();
+
+			var Keys = Old._Stats.Keys;
+
+			foreach( string KeyValue in Keys )
+			{
+				Int64 OldValue = Old[KeyValue];
+				Int64 NewValue = New[KeyValue];
+
+				Diff[KeyValue] = NewValue - OldValue;
+			}
+
+			return Diff;
+		}
+
+		/// <summary> Creates a new copy of this class. </summary>
+		public FMemoryAllocationStatsV4 DeepCopy()
+		{
+			FMemoryAllocationStatsV4 Copy = new FMemoryAllocationStatsV4(this);
+			return Copy;
+		}
+
+		/// <summary> Converts this memory allocation statistics to its equivalent string representation. </summary>
+		public override string ToString()
+		{
+			StringBuilder StrBuilder = new StringBuilder( 1024 );
+
+			foreach( KeyValuePair<string, Int64> KeyPairValue in _Stats )
+			{
+				StrBuilder.AppendLine( string.Format( "	{0}: {1}", KeyPairValue.Key, MainWindow.FormatSizeString2( KeyPairValue.Value ) ) );
+			}
+
+			return StrBuilder.ToString();
+		}
+
+		
+	};
+
 	/// <summary>
-    /// Variable sized token emitted by capture code. The parsing code ReadNextToken deals with this and updates
-    /// internal data. The calling code is responsible for only accessing member variables associated with the type.
-    /// </summary>
-    public class FStreamToken//@TODO: Can I turn this back into a struct?
-    {
-        // Parsing configuration
+	/// Variable sized token emitted by capture code. The parsing code ReadNextToken deals with this and updates
+	/// internal data. The calling code is responsible for only accessing member variables associated with the type.
+	/// </summary>
+	public class FStreamToken//@TODO: Can I turn this back into a struct?
+	{
+		// Parsing configuration
 
 		/// <summary> Mask of pointer field to get a real pointer (the two LSB are type fields, and the top bits may be a pool index. </summary>
-        public static UInt64 PointerMask = 0xFFFFFFFFFFFFFFFCUL;
+		public static UInt64 PointerMask = 0xFFFFFFFFFFFFFFFCUL;
 
-        public const UInt64 TypeMask = 0x3UL;
-        public static int PoolShift = 64;
+		public const UInt64 TypeMask = 0x3UL;
+		public static int PoolShift = 64;
 
-		/// <summary> Whether to decote the script callstacks. </summary>
-        public static bool bDecodeScriptCallstacks;
+		/// <summary> Whether to decode the script callstacks. </summary>
+		public static bool bDecodeScriptCallstacks;
 
-		/// <summary> Version of the stream. </summary>
-        public static uint Version;
+		/// <summary> Version of the stream, same as FProfileDataHeader.Version  </summary>
+		public static uint Version;
 
 		/// <summary> Position in the stream. </summary>
 		public ulong StreamIndex = 0;
 
 		/// <summary> Type of token. </summary>
-        public EProfilingPayloadType Type;
+		public EProfilingPayloadType Type;
 
 		/// <summary> Subtype of token if it's of TYPE_Other. </summary>
-        public EProfilingPayloadSubType SubType;
+		public EProfilingPayloadSubType SubType;
 
 		/// <summary> Pointer in the case of alloc / free. </summary>
-        public UInt64 Pointer;
+		public UInt64 Pointer;
 
 		/// <summary> Old pointer in the case of realloc. </summary>
-        public UInt64 OldPointer;
+		public UInt64 OldPointer;
 
 		/// <summary> New pointer in the case of realloc. </summary>
-        public UInt64 NewPointer;
+		public UInt64 NewPointer;
 
 		/// <summary> Index into callstack array. </summary>
-        public Int32 CallStackIndex;
+		public Int32 CallStackIndex;
 
 		/// <summary> Size of allocation in alloc / realloc case. </summary>
-        public Int32 Size;
+		public Int32 Size;
 
 		/// <summary> Payload if type is TYPE_Other. </summary>
-        public UInt32 Payload;
+		public UInt32 Payload;
 
 		/// <summary> Payload data if type is TYPE_Other and subtype is SUBTYPE_SnapshotMarker or SUBTYPE_TextMarker. Index into array of unique names. </summary>
-        public int TextIndex;
+		public int TextIndex;
 
 		/// <summary> Payload data if type is TYPE_Other and subtype is SUBTYPE_FrameTimeMarker. Current delta time in seconds. </summary>
-        public float DeltaTime;
+		public float DeltaTime;
 
 		/// <summary> Total time, increased every time DeltaTime has been read. </summary>
 		public float TotalTime = 0.0f;
@@ -413,22 +479,25 @@ namespace MemoryProfiler2
 		public float ElapsedTime = 0.0f;
 
 		/// <summary> Memory pool. </summary>
-        public EMemoryPool Pool;
+		public EMemoryPool Pool;
 
 		/// <summary> Platform dependent memory metrics. </summary>
-        public List<long> Metrics = new List<long>();
+		public List<long> Metrics = new List<long>();
 
 		/// <summary> A list of indices into the name table, one for each loaded level including persistent level. </summary>
-        public List<int> LoadedLevels = new List<int>();
+		public List<int> LoadedLevels = new List<int>();
+
+		/// <summary> Generic memory allocation stats for V3. </summary>
+		public FMemoryAllocationStatsV3 MemoryAllocationStats3 = new FMemoryAllocationStatsV3();
 
-		/// <summary> Generic memory allocation stats. </summary>
-		public FMemoryAllocationStats MemoryAllocationStats = new FMemoryAllocationStats();
+		/// <summary> Generic memory allocation stats for V4. </summary>
+		public FMemoryAllocationStatsV4 MemoryAllocationStats4 = new FMemoryAllocationStatsV4();
 
 		/// <summary> Index into script callstack array. </summary>
-        public int ScriptCallstackIndex;
+		public int ScriptCallstackIndex;
 
 		/// <summary> Index into script-object type array. </summary>
-        public int ScriptObjectTypeIndex;
+		public int ScriptObjectTypeIndex;
 
 		/// <summary> Reads a script callstack. </summary>
 		/// <param name="BinaryStream"> Stream to serialize data from </param>
@@ -479,15 +548,15 @@ namespace MemoryProfiler2
 
 		/// <summary> Reads platform dependent memory metrics. </summary>
 		/// <param name="BinaryStream"> Stream to serialize data from </param>
-        void ReadMetrics(BinaryReader BinaryStream)
-        {
-            // Read the metrics 
-            int NumMetrics = BinaryStream.ReadByte();
-            for (int i = 0; i < NumMetrics; i++)
-            {
-                Metrics.Add(BinaryStream.ReadInt64());
-            }
-        }
+		void ReadMetrics(BinaryReader BinaryStream)
+		{
+			// Read the metrics 
+			int NumMetrics = BinaryStream.ReadByte();
+			for (int i = 0; i < NumMetrics; i++)
+			{
+				Metrics.Add(BinaryStream.ReadInt64());
+			}
+		}
 
 		/// <summary> Reads names of all loaded levels. </summary>
 		/// <param name="BinaryStream"> Stream to serialize data from </param>
@@ -507,43 +576,61 @@ namespace MemoryProfiler2
 		/// <param name="BinaryStream"> Stream to serialize data from </param>
 		private void ReadMemoryAllocationsStats( BinaryReader BinaryStream )
 		{
-			int StatsNum = BinaryStream.ReadByte();
+			if( Version >= 4 )
+			{
+				// Read Platform Memory and Allocator Stats
+				// @see FMallocProfiler::WriteMemoryAllocationStats
+				int StatsNum = BinaryStream.ReadByte();
 
-			FieldInfo[] FieldInfos = MemoryAllocationStats.GetType().GetFields();
-			int PropertiesNum = FieldInfos.Length;
-			Debug.Assert( StatsNum == PropertiesNum );
+				for( int StatIndex = 0; StatIndex < StatsNum; StatIndex++ )
+				{
+					int NameIndex = BinaryStream.ReadInt32();
+					Int64 StatValue = BinaryStream.ReadInt64();
+					string StatName = FStreamInfo.GlobalInstance.NameArray[NameIndex];
 
-			for( int StatIndex = 0; StatIndex < StatsNum; StatIndex++ )
-			{
-				Int64 Value = BinaryStream.ReadInt64();
-				FieldInfos[ StatIndex ].SetValue( MemoryAllocationStats, Value );
+					MemoryAllocationStats4[StatName] = StatValue;
+				}
 			}
+			else
+			{
+				int StatsNum = BinaryStream.ReadByte();
+
+				FieldInfo[] FieldInfos = MemoryAllocationStats3.GetType().GetFields();
+				int PropertiesNum = FieldInfos.Length;
+				Debug.Assert( StatsNum == PropertiesNum );
+
+				for( int StatIndex = 0; StatIndex < StatsNum; StatIndex++ )
+				{
+					Int64 Value = BinaryStream.ReadInt64();
+					FieldInfos[StatIndex].SetValue( MemoryAllocationStats3, Value );
+				}
+			}			
 		}
 
 		/// <summary> Updates the token with data read from passed in stream and returns whether we've reached the end. </summary>
 		/// <param name="BinaryStream"> Stream to serialize data from </param>
 		public bool ReadNextToken( BinaryReader BinaryStream )
-        {
-            bool bReachedEndOfStream = false;
+		{
+			bool bReachedEndOfStream = false;
 
 			// Initialize to defaults.
 			SubType = EProfilingPayloadSubType.SUBTYPE_Unknown;
 			TextIndex = -1;
 
-            // Read the pointer and convert to token type by looking at lowest 2 bits. Pointers are always
-            // 4 byte aligned so need to clear them again after the conversion.
-            UInt64 RawPointerData = BinaryStream.ReadUInt64();
+			// Read the pointer and convert to token type by looking at lowest 2 bits. Pointers are always
+			// 4 byte aligned so need to clear them again after the conversion.
+			UInt64 RawPointerData = BinaryStream.ReadUInt64();
 
 			Pool = EMemoryPool.MEMPOOL_Main;
-            Type = (EProfilingPayloadType)(RawPointerData & TypeMask);
-            Pointer = RawPointerData & PointerMask;
+			Type = (EProfilingPayloadType)(RawPointerData & TypeMask);
+			Pointer = RawPointerData & PointerMask;
 
-            Metrics.Clear();
-            LoadedLevels.Clear();
-			MemoryAllocationStats.Zero();
+			Metrics.Clear();
+			LoadedLevels.Clear();
+			MemoryAllocationStats3.Zero();
 			CallStackIndex = -1;
-            ScriptCallstackIndex = -1;
-            ScriptObjectTypeIndex = -1;
+			ScriptCallstackIndex = -1;
+			ScriptObjectTypeIndex = -1;
 
 			NewPointer = 0;
 			OldPointer = 0;
@@ -551,11 +638,11 @@ namespace MemoryProfiler2
 			Payload = 0;
 			DeltaTime = -1.0f;
 
-            // Serialize based on token type.
+			// Serialize based on token type.
 			switch( Type )
-            {
-                // Malloc
-                case EProfilingPayloadType.TYPE_Malloc:
+			{
+				// Malloc
+				case EProfilingPayloadType.TYPE_Malloc:
 				{
 					// Get the call stack index.
 					CallStackIndex = BinaryStream.ReadInt32();
@@ -576,9 +663,9 @@ namespace MemoryProfiler2
 				}
 				
 
-                // Free
-                case EProfilingPayloadType.TYPE_Free:
-                {
+				// Free
+				case EProfilingPayloadType.TYPE_Free:
+				{
 					break;
 				}
 
@@ -600,7 +687,7 @@ namespace MemoryProfiler2
 				}
 				
 
-                // Other
+				// Other
 				case EProfilingPayloadType.TYPE_Other:
 				{
 					SubType = ( EProfilingPayloadSubType )BinaryStream.ReadInt32();
@@ -612,12 +699,16 @@ namespace MemoryProfiler2
 						// End of stream!
 						case EProfilingPayloadSubType.SUBTYPE_EndOfStreamMarker:
 						{
-							if( Version > 2 )
+							if( Version >= 4 )
+							{
+								ReadMemoryAllocationsStats( BinaryStream );
+								ReadLoadedLevels( BinaryStream );
+							}
+							else
 							{
 								ReadMemoryAllocationsStats( BinaryStream );
 								ReadLoadedLevels( BinaryStream );
 								ReadMetrics( BinaryStream );
-
 							}
 							bReachedEndOfStream = true;
 							break;
@@ -640,7 +731,12 @@ namespace MemoryProfiler2
 						{
 							TextIndex = ( int )Payload;
 
-							if( Version > 2 )
+							if( Version >= 4 )
+							{
+								ReadMemoryAllocationsStats( BinaryStream );
+								ReadLoadedLevels( BinaryStream );
+							}
+							else
 							{
 								ReadMemoryAllocationsStats( BinaryStream );
 								ReadLoadedLevels( BinaryStream );
@@ -693,7 +789,7 @@ namespace MemoryProfiler2
 				}
 			}
 
-            return !bReachedEndOfStream;
-        }
-    }
+			return !bReachedEndOfStream;
+		}
+	}
 }
diff --git a/Engine/Source/Programs/MemoryProfiler2/TimeLineView.cs b/Engine/Source/Programs/MemoryProfiler2/TimeLineView.cs
--- a/Engine/Source/Programs/MemoryProfiler2/TimeLineView.cs
+++ b/Engine/Source/Programs/MemoryProfiler2/TimeLineView.cs
@@ -12,6 +12,9 @@ namespace MemoryProfiler2
 {
 	public static class FTimeLineChartView
 	{
+		/// <summary> Memory allocated calculated by the profiler. </summary>
+		static public string AllocatedMemory = "Allocated Memory";
+
 		/// <summary> Reference to the main memory profiler window. </summary>
 		private static MainWindow OwnerWindow;
 
@@ -32,7 +35,7 @@ namespace MemoryProfiler2
 			}
 		}
 
-		private static void AddSeries( Chart TimeLineChart, string SeriesName, FStreamSnapshot Snapshot, ESliceTypes SliceType, bool bVisible )
+		private static void AddSeriesV3( Chart TimeLineChart, string SeriesName, FStreamSnapshot Snapshot, ESliceTypesV3 SliceType, bool bVisible )
 		{
 			Series ChartSeries = TimeLineChart.Series[ SeriesName ];
 			ChartSeries.Enabled = bVisible;
@@ -44,43 +47,130 @@ namespace MemoryProfiler2
 				foreach( FMemorySlice Slice in Snapshot.OverallMemorySlice )
 				{
 					DataPoint Point = new DataPoint();
-					Point.YValues[ 0 ] = Slice.GetSliceInfo( SliceType ) / ( 1024.0 * 1024.0 );
+					Point.YValues[ 0 ] = Slice.GetSliceInfoV3( SliceType ) / ( 1024.0 * 1024.0 );
 
 					ChartSeries.Points.Add( Point );
 				}
 			}
 		}
 
-		public static void ParseSnapshot( Chart TimeLineChart, FStreamSnapshot Snapshot )
+		private static void AddSeriesV4( Chart TimeLineChart, string SeriesName, FStreamSnapshot Snapshot, ESliceTypesV4 SliceType, bool bVisible )
 		{
-			// Progress bar.
-			OwnerWindow.UpdateStatus( "Updating time line view for " + OwnerWindow.CurrentFilename );
+			Series ChartSeries = TimeLineChart.Series[SeriesName];
+			ChartSeries.Enabled = bVisible;
 
-			TimeLineChart.Annotations.Clear();
+			if( ChartSeries != null )
+			{
+				ChartSeries.Points.Clear();
+
+				foreach( FMemorySlice Slice in Snapshot.OverallMemorySlice )
+				{
+					DataPoint Point = new DataPoint();
+					Point.YValues[0] = Slice.GetSliceInfoV4( SliceType ) / ( 1024.0 * 1024.0 );
+
+					ChartSeries.Points.Add( Point );
+				}
+			}
+		}
+
+		private static void AddEmptySeries( Chart TimeLineChart, string SeriesName, int NumPoints )
+		{
+			Series ChartSeries = TimeLineChart.Series[SeriesName];
+			ChartSeries.Enabled = false;
+
+			if( ChartSeries != null )
+			{
+				ChartSeries.Points.Clear();
 
-			bool bIsXbox360 = FStreamInfo.GlobalInstance.Platform == EPlatformType.Xbox360;
-			bool bIsPS3 = FStreamInfo.GlobalInstance.Platform == EPlatformType.PS3;
+				for( int Index = 0; Index < NumPoints; Index++ )
+				{
+					ChartSeries.Points.Add( 0 );
+				}
+			}
+		}
 
-			AddSeries( TimeLineChart, "Allocated Memory", Snapshot, ESliceTypes.TotalUsed, true );
+		private static void AddFakeSeries( Chart TimeLineChart, string SeriesName, int NumPoints, int Value )
+		{
+			Series ChartSeries = TimeLineChart.Series[SeriesName];
+			ChartSeries.Enabled = true;
 
-			AddSeries( TimeLineChart, "Image Size", Snapshot, ESliceTypes.ImageSize, bIsXbox360 || bIsPS3 );
-			AddSeries( TimeLineChart, "OS Overhead", Snapshot, ESliceTypes.OSOverhead, bIsXbox360 );
+			if( ChartSeries != null )
+			{
+				ChartSeries.Points.Clear();
 
-			AddSeries( TimeLineChart, "Virtual Used", Snapshot, ESliceTypes.CPUUsed, true );
-			AddSeries( TimeLineChart, "Virtual Slack", Snapshot, ESliceTypes.CPUSlack, true );
-			AddSeries( TimeLineChart, "Virtual Waste", Snapshot, ESliceTypes.CPUWaste, true );
-			AddSeries( TimeLineChart, "Physical Used", Snapshot, ESliceTypes.GPUUsed, bIsPS3 || bIsXbox360 );
-			AddSeries( TimeLineChart, "Physical Slack", Snapshot, ESliceTypes.GPUSlack, bIsXbox360 );
-			AddSeries( TimeLineChart, "Physical Waste", Snapshot, ESliceTypes.GPUWaste, bIsPS3 || bIsXbox360 );
+				for( int Index = 0; Index < NumPoints; Index++ )
+				{
+					ChartSeries.Points.Add( Value );
+				}
+			}
+		}
+
+		public static void ParseSnapshot( Chart TimeLineChart, FStreamSnapshot Snapshot )
+		{
+			// Progress bar.
+			OwnerWindow.UpdateStatus( "Updating time line view for " + OwnerWindow.CurrentFilename );
 
-			AddSeries( TimeLineChart, "Host Used", Snapshot, ESliceTypes.HostUsed, bIsPS3 );
-			AddSeries( TimeLineChart, "Host Slack", Snapshot, ESliceTypes.HostSlack, bIsPS3 );
-			AddSeries( TimeLineChart, "Host Waste", Snapshot, ESliceTypes.HostWaste, bIsPS3 );
+			TimeLineChart.Annotations.Clear();
+
+			if( FStreamToken.Version >= 4 )
+			{
+				AddSeriesV4( TimeLineChart, FTimeLineChartView.AllocatedMemory, Snapshot, ESliceTypesV4.OverallAllocatedMemory, true );
+				AddSeriesV4( TimeLineChart, FMemoryAllocationStatsV4.PlatformUsedPhysical, Snapshot, ESliceTypesV4.PlatformUsedPhysical, true );
+
+				AddSeriesV4( TimeLineChart, FMemoryAllocationStatsV4.MemoryProfilingOverhead, Snapshot, ESliceTypesV4.MemoryProfilingOverhead, true );
+				AddSeriesV4( TimeLineChart, FMemoryAllocationStatsV4.BinnedWasteCurrent, Snapshot, ESliceTypesV4.BinnedWasteCurrent, true );
+				AddSeriesV4( TimeLineChart, FMemoryAllocationStatsV4.BinnedSlackCurrent, Snapshot, ESliceTypesV4.BinnedSlackCurrent, true );
+				AddSeriesV4( TimeLineChart, FMemoryAllocationStatsV4.BinnedUsedCurrent, Snapshot, ESliceTypesV4.BinnedUsedCurrent, true );
+				
+				AddEmptySeries( TimeLineChart, "Image Size", Snapshot.OverallMemorySlice.Count );
+				AddEmptySeries( TimeLineChart, "OS Overhead", Snapshot.OverallMemorySlice.Count );
+
+				AddEmptySeries( TimeLineChart, "Virtual Used", Snapshot.OverallMemorySlice.Count );
+				AddEmptySeries( TimeLineChart, "Virtual Slack", Snapshot.OverallMemorySlice.Count );
+				AddEmptySeries( TimeLineChart, "Virtual Waste", Snapshot.OverallMemorySlice.Count );
+				AddEmptySeries( TimeLineChart, "Physical Used", Snapshot.OverallMemorySlice.Count );
+				AddEmptySeries( TimeLineChart, "Physical Slack", Snapshot.OverallMemorySlice.Count );
+				AddEmptySeries( TimeLineChart, "Physical Waste", Snapshot.OverallMemorySlice.Count );
+
+				AddEmptySeries( TimeLineChart, "Host Used", Snapshot.OverallMemorySlice.Count );
+				AddEmptySeries( TimeLineChart, "Host Slack", Snapshot.OverallMemorySlice.Count );
+				AddEmptySeries( TimeLineChart, "Host Waste", Snapshot.OverallMemorySlice.Count );
+			}
+			else
+			{
+				bool bIsXbox360 = FStreamInfo.GlobalInstance.Platform == EPlatformType.Xbox360;
+				bool bIsPS3 = FStreamInfo.GlobalInstance.Platform == EPlatformType.PS3;
+
+				AddSeriesV3( TimeLineChart, FTimeLineChartView.AllocatedMemory, Snapshot, ESliceTypesV3.OverallAllocatedMemory, true );
+
+				// Total used will be displayed as "Used Physical".
+				AddSeriesV3( TimeLineChart, FMemoryAllocationStatsV4.PlatformUsedPhysical, Snapshot, ESliceTypesV3.TotalUsed, true );
+
+				AddSeriesV3( TimeLineChart, "Image Size", Snapshot, ESliceTypesV3.ImageSize, bIsXbox360 || bIsPS3 );
+				AddSeriesV3( TimeLineChart, "OS Overhead", Snapshot, ESliceTypesV3.OSOverhead, bIsXbox360 );
+
+				AddSeriesV3( TimeLineChart, "Virtual Used", Snapshot, ESliceTypesV3.CPUUsed, true );
+				AddSeriesV3( TimeLineChart, "Virtual Slack", Snapshot, ESliceTypesV3.CPUSlack, true );
+				AddSeriesV3( TimeLineChart, "Virtual Waste", Snapshot, ESliceTypesV3.CPUWaste, true );
+				AddSeriesV3( TimeLineChart, "Physical Used", Snapshot, ESliceTypesV3.GPUUsed, bIsPS3 || bIsXbox360 );
+				AddSeriesV3( TimeLineChart, "Physical Slack", Snapshot, ESliceTypesV3.GPUSlack, bIsXbox360 );
+				AddSeriesV3( TimeLineChart, "Physical Waste", Snapshot, ESliceTypesV3.GPUWaste, bIsPS3 || bIsXbox360 );
+
+				AddSeriesV3( TimeLineChart, "Host Used", Snapshot, ESliceTypesV3.HostUsed, bIsPS3 );
+				AddSeriesV3( TimeLineChart, "Host Slack", Snapshot, ESliceTypesV3.HostSlack, bIsPS3 );
+				AddSeriesV3( TimeLineChart, "Host Waste", Snapshot, ESliceTypesV3.HostWaste, bIsPS3 );
+
+				//AddEmptySeries( TimeLineChart, FMemoryAllocationStatsV4.PlatformUsedPhysical, Snapshot.OverallMemorySlice.Count );
+				AddEmptySeries( TimeLineChart, FMemoryAllocationStatsV4.BinnedWasteCurrent, Snapshot.OverallMemorySlice.Count );
+				AddEmptySeries( TimeLineChart, FMemoryAllocationStatsV4.BinnedSlackCurrent, Snapshot.OverallMemorySlice.Count );
+				AddEmptySeries( TimeLineChart, FMemoryAllocationStatsV4.BinnedUsedCurrent, Snapshot.OverallMemorySlice.Count );
+				AddEmptySeries( TimeLineChart, FMemoryAllocationStatsV4.MemoryProfilingOverhead, Snapshot.OverallMemorySlice.Count );
+			}	
 		}
 
 		public static bool AddCustomSnapshot( Chart TimeLineChart, CursorEventArgs Event )
 		{
-			if( TimeLineChart.Series["Allocated Memory"].Points.Count == 0 )
+			if( TimeLineChart.Series[FTimeLineChartView.AllocatedMemory].Points.Count == 0 )
 			{
 				return false;
 			}
