#UE4
- TTP# 331805 (MemoryProfiler2 no longer compiles, is no longer checked in, fixed)

diff --git a/Engine/Source/Programs/MemoryProfiler2/Android/Android.cs b/Engine/Source/Programs/MemoryProfiler2/Android/Android.cs
deleted file mode 100644
diff --git a/Engine/Source/Programs/MemoryProfiler2/IPhone/IPhone.cs b/Engine/Source/Programs/MemoryProfiler2/IPhone/IPhone.cs
deleted file mode 100644
diff --git a/Engine/Source/Programs/MemoryProfiler2/MemoryProfiler2.csproj b/Engine/Source/Programs/MemoryProfiler2/MemoryProfiler2.csproj
--- a/Engine/Source/Programs/MemoryProfiler2/MemoryProfiler2.csproj
+++ b/Engine/Source/Programs/MemoryProfiler2/MemoryProfiler2.csproj
@@ -62,13 +62,7 @@
     <PlatformTarget>x64</PlatformTarget>
   </PropertyGroup>
   <ItemGroup>
-    <Reference Include="AndroidTools_x64">
-      <HintPath>Android\AndroidTools_x64.dll</HintPath>
-    </Reference>
     <Reference Include="EnvDTE, Version=8.0.0.0, Culture=neutral, PublicKeyToken=b03f5f7f11d50a3a" />
-    <Reference Include="IPhoneTools_x64">
-      <HintPath>IPhone\IPhoneTools_x64.dll</HintPath>
-    </Reference>
     <Reference Include="System" />
     <Reference Include="System.Core">
       <RequiredTargetFramework>3.5</RequiredTargetFramework>
@@ -186,26 +180,6 @@
     <Compile Include="StreamToken.cs" />
     <Compile Include="TimeLineView.cs" />
   </ItemGroup>
-  <Choose>
-    <When Condition="Exists('Android\Android.cs')">
-      <ItemGroup>
-        <Compile Include="Android\Android.cs" />
-      </ItemGroup>
-    </When>
-  </Choose>
-  <Choose>
-    <When Condition="Exists('PS3\PS3.cs')" />
-  </Choose>
-  <Choose>
-    <When Condition="Exists('IPhone\IPhone.cs')">
-      <ItemGroup>
-        <Compile Include="IPhone\IPhone.cs" />
-      </ItemGroup>
-    </When>
-  </Choose>
-  <Choose>
-    <When Condition="Exists('Xbox360\Xbox360.cs')" />
-  </Choose>
   <ItemGroup>
     <BootstrapperPackage Include="Microsoft.Net.Client.3.5">
       <Visible>False</Visible>
diff --git a/Engine/Source/Programs/MemoryProfiler2/PS3/PS3.cs b/Engine/Source/Programs/MemoryProfiler2/PS3/PS3.cs
deleted file mode 100644
diff --git a/Engine/Source/Programs/MemoryProfiler2/Xbox360/Xbox360.cs b/Engine/Source/Programs/MemoryProfiler2/Xbox360/Xbox360.cs
deleted file mode 100644
