MemoryProfiler2 - Updated comments

diff --git a/Engine/Source/Programs/MemoryProfiler2/StreamToken.cs b/Engine/Source/Programs/MemoryProfiler2/StreamToken.cs
--- a/Engine/Source/Programs/MemoryProfiler2/StreamToken.cs
+++ b/Engine/Source/Programs/MemoryProfiler2/StreamToken.cs
@@ -309,7 +309,7 @@ namespace MemoryProfiler2
 	/// </summary>
 	public class FMemoryAllocationStatsV4
 	{
-		// Names of generic platform memory stats. @see FGenericPlatformMemoryStats
+		// Names of generic platform memory stats. @see FGenericPlatformMemory::GetStatsForMallocProfiler / GenericPlatformMemory.cpp
 		
 		/// <summary> The amount of actual physical memory, in bytes. </summary>
 		static public string PlatformTotalPhysical = "Total Physical";
@@ -343,7 +343,7 @@ namespace MemoryProfiler2
 
 		static public string MemoryProfilingOverhead = "Memory Profiling Overhead";
 
-		// Names of malloc binned memory stats. @see FMallocBinned.GetAllocatorStats 
+		// Names of malloc binned memory stats. @see FMallocBinned.GetAllocatorStats / MallocBinned.h
 
 		static public string BinnedWasteCurrent = "Binned Waste Current";
 		static public string BinnedUsedCurrent = "Binned Used Current";
