MallocProfiler2 - Better handling of the missing stats values, updated a link to the documentation

diff --git a/Engine/Source/Programs/MemoryProfiler2/MainWindow.cs b/Engine/Source/Programs/MemoryProfiler2/MainWindow.cs
--- a/Engine/Source/Programs/MemoryProfiler2/MainWindow.cs
+++ b/Engine/Source/Programs/MemoryProfiler2/MainWindow.cs
@@ -759,7 +759,7 @@ namespace MemoryProfiler2
 		private void QuickStartMenuClick( object sender, EventArgs e )
 		{
 			//MessageBox.Show( "Define USE_MALLOC_PROFILER 1 in UMemoryDefines.h\n\nWhile running, type SNAPSHOTMEMORY at suitable intervals. At the end of profiling, type DUMPALLOCSTOFILE to save the .mprof file.\n\nLoad in the .mprof file from <Name>Game/Profiling/<timestamp>/<Name>-<timestamp>.mprof", "Memory Profiler Quick Start", MessageBoxButtons.OK, MessageBoxIcon.Information );
-			System.Diagnostics.Process.Start( "https://udn.epicgames.com/Three/MemoryProfiler" );
+			System.Diagnostics.Process.Start( "https://udn.unrealengine.com/docs/ue3/INT/Engine/Subsystems/Memory/Profiling/Tools/Profiler/index.html" );
 		}
 
 		/// <summary> Opens the options dialog. </summary>
diff --git a/Engine/Source/Programs/MemoryProfiler2/StreamToken.cs b/Engine/Source/Programs/MemoryProfiler2/StreamToken.cs
--- a/Engine/Source/Programs/MemoryProfiler2/StreamToken.cs
+++ b/Engine/Source/Programs/MemoryProfiler2/StreamToken.cs
@@ -365,7 +365,9 @@ namespace MemoryProfiler2
 		{
 			get
 			{
-				return _Stats[StatName];
+				Int64 Value = 0;
+				_Stats.TryGetValue(StatName, out Value);
+				return Value;
 			}
 			set
 			{
