MallocProfiler - Changed uint32 to uint64 in the FProfilerHeader to support profiler files larger than 4.0GB, increased memory profiler version to 5, updated C# app, misc compilation fixes

diff --git a/Engine/Source/Programs/MemoryProfiler2/ProfileDataHeader.cs b/Engine/Source/Programs/MemoryProfiler2/ProfileDataHeader.cs
--- a/Engine/Source/Programs/MemoryProfiler2/ProfileDataHeader.cs
+++ b/Engine/Source/Programs/MemoryProfiler2/ProfileDataHeader.cs
@@ -19,7 +19,7 @@ namespace MemoryProfiler2
 		/** Version number to detect version mismatches.			*/
 		public UInt32 Version;
 		/** Platform that was captured.								*/
-        public EPlatformType Platform;
+		public EPlatformType Platform;
 		/** Platform that was captured.								*/
 		public string PlatformName;
 		/** Whether symbol information was serialized.				*/
@@ -28,37 +28,36 @@ namespace MemoryProfiler2
 		public string ExecutableName;
 
 		/** Offset in file for name table.							*/
-		public UInt32 NameTableOffset;
+		public UInt64 NameTableOffset;
 		/** Number of name table entries.							*/
-		public UInt32 NameTableEntries;
+		public UInt64 NameTableEntries;
 
 		/** Offset in file for callstack address table.				*/
-		public UInt32 CallStackAddressTableOffset;
+		public UInt64 CallStackAddressTableOffset;
 		/** Number of callstack address entries.					*/
-		public UInt32 CallStackAddressTableEntries;
+		public UInt64 CallStackAddressTableEntries;
 
 		/** Offset in file for callstack table.						*/
-		public UInt32 CallStackTableOffset;
+		public UInt64 CallStackTableOffset;
 		/** Number of callstack entries.							*/
-		public UInt32 CallStackTableEntries;
+		public UInt64 CallStackTableEntries;
 		/** The file offset for module information.					*/
-		public uint ModulesOffset;
+		public UInt64 ModulesOffset;
 		/** The number of module entries.							*/
-		public uint ModuleEntries;
-
+		public UInt64 ModuleEntries;
 		/** Number of data files the stream spans.					*/
-		public UInt32 NumDataFiles;
+		public UInt64 NumDataFiles;
 
 		// New in version 3 files
 
 		/** Offset in file for stript callstack table. Requires version 3 or later. */
-        public UInt32 ScriptCallstackTableOffset = UInt32.MaxValue;
+		public UInt32 ScriptCallstackTableOffset = UInt32.MaxValue;
 
 		/** Offset in file for script name table. Requires version 3 or later. */
-        public UInt32 ScriptNameTableOffset = UInt32.MaxValue;
+		public UInt32 ScriptNameTableOffset = UInt32.MaxValue;
 
 		/** Can/should script callstacks be converted into readable names? */
-        public bool bDecodeScriptCallstacks = false;
+		public bool bDecodeScriptCallstacks = false;
 
 		/**
 		 * Constructor, serializing header from passed in stream.
@@ -91,36 +90,59 @@ namespace MemoryProfiler2
 				// Whether symbol information was serialized.
 				bShouldSerializeSymbolInfo = BinaryStream.ReadUInt32() == 0 ? false : true;
 
-				// Name table offset in file and number of entries.
-				NameTableOffset = BinaryStream.ReadUInt32();
-				NameTableEntries = BinaryStream.ReadUInt32();
+				if( Version >= 5 )
+				{
+					// Name table offset in file and number of entries.
+					NameTableOffset = BinaryStream.ReadUInt64();
+					NameTableEntries = BinaryStream.ReadUInt64();
+
+					// CallStack address table offset and number of entries.
+					CallStackAddressTableOffset = BinaryStream.ReadUInt64();
+					CallStackAddressTableEntries = BinaryStream.ReadUInt64();
+
+					// CallStack table offset and number of entries.
+					CallStackTableOffset = BinaryStream.ReadUInt64();
+					CallStackTableEntries = BinaryStream.ReadUInt64();
+
+					ModulesOffset = BinaryStream.ReadUInt64();
+					ModuleEntries = BinaryStream.ReadUInt64();
+
+					NumDataFiles = 1;
+				}
+				else
+				{
+					// Name table offset in file and number of entries.
+					NameTableOffset = BinaryStream.ReadUInt32();
+					NameTableEntries = BinaryStream.ReadUInt32();
 
-				// CallStack address table offset and number of entries.
-				CallStackAddressTableOffset = BinaryStream.ReadUInt32();
-				CallStackAddressTableEntries = BinaryStream.ReadUInt32();
+					// CallStack address table offset and number of entries.
+					CallStackAddressTableOffset = BinaryStream.ReadUInt32();
+					CallStackAddressTableEntries = BinaryStream.ReadUInt32();
 
-				// CallStack table offset and number of entries.
-				CallStackTableOffset = BinaryStream.ReadUInt32();
-				CallStackTableEntries = BinaryStream.ReadUInt32();
+					// CallStack table offset and number of entries.
+					CallStackTableOffset = BinaryStream.ReadUInt32();
+					CallStackTableEntries = BinaryStream.ReadUInt32();
 
-				ModulesOffset = BinaryStream.ReadUInt32();
-				ModuleEntries = BinaryStream.ReadUInt32();
+					ModulesOffset = BinaryStream.ReadUInt32();
+					ModuleEntries = BinaryStream.ReadUInt32();
 
-				// Number of data files the stream spans.
-				NumDataFiles = BinaryStream.ReadUInt32();
+					// Number of data files the stream spans.
+					NumDataFiles = BinaryStream.ReadUInt32();
+				}
+	
 
-                FStreamToken.Version = Version;
+				FStreamToken.Version = Version;
 				if( Version >= 4 )
 				{
 					// Just ignore
 				}
 				else if (Version > 2)
 				{
-	                ScriptCallstackTableOffset = BinaryStream.ReadUInt32();
-                    ScriptNameTableOffset = BinaryStream.ReadUInt32();
-                
-                    bDecodeScriptCallstacks = ScriptCallstackTableOffset != UInt32.MaxValue;
-                }
+					ScriptCallstackTableOffset = BinaryStream.ReadUInt32();
+					ScriptNameTableOffset = BinaryStream.ReadUInt32();
+				
+					bDecodeScriptCallstacks = ScriptCallstackTableOffset != UInt32.MaxValue;
+				}
 
 				// Name of executable.
 				ExecutableName = FStreamParser.ReadString( BinaryStream );
@@ -140,38 +162,38 @@ namespace MemoryProfiler2
 		}
 	}
 
-    // Mirrored in UnFile.h
-    [Flags]
-    public enum EPlatformType
-    {
-        Unknown = 0x00000000,
-        Windows = 0x00000001,
-        WindowsServer = 0x00000002,		// Windows platform dedicated server mode ("lean and mean" cooked as console without editor support)
-        Xbox360 = 0x00000004,
-        PS3 = 0x00000008,
-        Linux = 0x00000010,
-        MacOSX = 0x00000020,
-        WindowsConsole = 0x00000040,     // Windows platform cooked as console without editor support
-        IPhone = 0x00000080,
-        Android = 0x00000200,
-
-        // Combination Masks
-        /** PC platform types */
-        PC = Windows | WindowsServer | WindowsConsole | Linux | MacOSX,
-
-        /** Windows platform types */
-        AnyWindows = Windows | WindowsServer | WindowsConsole,
-
-        /** Console platform types */
+	// Mirrored in UnFile.h
+	[Flags]
+	public enum EPlatformType
+	{
+		Unknown = 0x00000000,
+		Windows = 0x00000001,
+		WindowsServer = 0x00000002,		// Windows platform dedicated server mode ("lean and mean" cooked as console without editor support)
+		Xbox360 = 0x00000004,
+		PS3 = 0x00000008,
+		Linux = 0x00000010,
+		MacOSX = 0x00000020,
+		WindowsConsole = 0x00000040,     // Windows platform cooked as console without editor support
+		IPhone = 0x00000080,
+		Android = 0x00000200,
+
+		// Combination Masks
+		/** PC platform types */
+		PC = Windows | WindowsServer | WindowsConsole | Linux | MacOSX,
+
+		/** Windows platform types */
+		AnyWindows = Windows | WindowsServer | WindowsConsole,
+
+		/** Console platform types */
 		Console = Xbox360 | PS3 | IPhone | Android,
 
-        /** Mobile platform types */
-        Mobile = IPhone | Android,
+		/** Mobile platform types */
+		Mobile = IPhone | Android,
 
-        /** Platforms with data that has been stripped during cooking */
-        Stripped = Console | WindowsServer | WindowsConsole,
+		/** Platforms with data that has been stripped during cooking */
+		Stripped = Console | WindowsServer | WindowsConsole,
 
-        /** Platforms who's vertex data can't be packed into 16-bit floats */
-        OpenGLES2 = IPhone | Android,
-    }
+		/** Platforms who's vertex data can't be packed into 16-bit floats */
+		OpenGLES2 = IPhone | Android,
+	}
 }
\ No newline at end of file
diff --git a/Engine/Source/Programs/MemoryProfiler2/StreamParser.cs b/Engine/Source/Programs/MemoryProfiler2/StreamParser.cs
--- a/Engine/Source/Programs/MemoryProfiler2/StreamParser.cs
+++ b/Engine/Source/Programs/MemoryProfiler2/StreamParser.cs
@@ -81,7 +81,7 @@ namespace MemoryProfiler2
 			}
 
             // Keep track of the current data file for multi-part recordings
-            int NextDataFile = 1;
+            UInt64 NextDataFile = 1;
 
 			// Initialize shared information across snapshots, namely names, callstacks and addresses.
 			FStreamInfo.GlobalInstance.Initialize( Header );
@@ -90,23 +90,23 @@ namespace MemoryProfiler2
 			long TokenStreamOffset = ParserFileStream.Position;
 
 			// Seek to name table and serialize it.
-			ParserFileStream.Seek(Header.NameTableOffset,SeekOrigin.Begin);
-			for(int NameIndex = 0;NameIndex < Header.NameTableEntries;NameIndex++)
+			ParserFileStream.Seek((Int64)Header.NameTableOffset,SeekOrigin.Begin);
+			for(UInt64 NameIndex = 0;NameIndex < Header.NameTableEntries;NameIndex++)
 			{
 				FStreamInfo.GlobalInstance.NameArray.Add( ReadString( BinaryStream ) );
 			}
 
 			// Seek to callstack address array and serialize it.                
-			ParserFileStream.Seek(Header.CallStackAddressTableOffset,SeekOrigin.Begin);
-			for(int AddressIndex = 0;AddressIndex < Header.CallStackAddressTableEntries;AddressIndex++)
+			ParserFileStream.Seek( (Int64)Header.CallStackAddressTableOffset, SeekOrigin.Begin );
+			for(UInt64 AddressIndex = 0;AddressIndex < Header.CallStackAddressTableEntries;AddressIndex++)
 			{
                 FStreamInfo.GlobalInstance.CallStackAddressArray.Add(new FCallStackAddress(BinaryStream, Header.bShouldSerializeSymbolInfo));
 			}
 
 			// Seek to callstack array and serialize it.
-			ParserFileStream.Seek(Header.CallStackTableOffset,SeekOrigin.Begin);
+			ParserFileStream.Seek( (Int64)Header.CallStackTableOffset, SeekOrigin.Begin );
 
-			for(int CallStackIndex = 0;CallStackIndex < Header.CallStackTableEntries;CallStackIndex++)
+			for(UInt64 CallStackIndex = 0;CallStackIndex < Header.CallStackTableEntries;CallStackIndex++)
 			{
                 FStreamInfo.GlobalInstance.CallStackArray.Add(new FCallStack(BinaryStream));
 			}
@@ -271,12 +271,12 @@ namespace MemoryProfiler2
             int SnapshotIndex = 0;
 
             // Figure out the progress scale
-            long StartOfMetadata = Math.Min(Header.NameTableOffset, Header.CallStackAddressTableOffset);
+            UInt64 StartOfMetadata = Math.Min(Header.NameTableOffset, Header.CallStackAddressTableOffset);
             StartOfMetadata = Math.Min(StartOfMetadata, Header.CallStackTableOffset);
             StartOfMetadata = Math.Min(StartOfMetadata, Header.ModulesOffset);
 
-            long ProgressInterval = (StartOfMetadata - TokenStreamOffset) / 100;
-            double ProgressScaleFactor = 100.0f / (StartOfMetadata - TokenStreamOffset);
+            long ProgressInterval = ((Int64)StartOfMetadata - TokenStreamOffset) / 100;
+			double ProgressScaleFactor = 100.0f / ( (Int64)StartOfMetadata - TokenStreamOffset );
             long NextProgressUpdate = TokenStreamOffset;
 
 			// Parse tokens till we reach the end of the stream.
@@ -490,7 +490,7 @@ namespace MemoryProfiler2
 								{
 									// Switch to the next file in the chain
 									ParserFileStream.Close();
-									BinaryStream = SwitchStreams( NextDataFile, FStreamInfo.GlobalInstance.FileName, bIsBigEndian, out ParserFileStream );
+									BinaryStream = SwitchStreams( (int)NextDataFile, FStreamInfo.GlobalInstance.FileName, bIsBigEndian, out ParserFileStream );
 
 									// Update variables used for reporting progress
 									TokenStreamOffset = 0;
@@ -902,8 +902,8 @@ namespace MemoryProfiler2
                 {
 					if( Header.ModuleEntries > 0 )
 					{
-						BinaryStream.BaseStream.Seek( Header.ModulesOffset, SeekOrigin.Begin );
-						FStreamInfo.GlobalInstance.ConsoleSymbolParser.ReadModuleInfo( BinaryStream, Header.ModuleEntries );
+						BinaryStream.BaseStream.Seek( (Int64)Header.ModulesOffset, SeekOrigin.Begin );
+						FStreamInfo.GlobalInstance.ConsoleSymbolParser.ReadModuleInfo( BinaryStream, (uint)Header.ModuleEntries );
 					}
 
                     string CurrentPath = MainMProfWindow.CurrentFilename.Substring(0, LastPathSeparator);
diff --git a/Engine/Source/Runtime/Core/Private/ProfilingDebugging/MallocProfiler.cpp b/Engine/Source/Runtime/Core/Private/ProfilingDebugging/MallocProfiler.cpp
--- a/Engine/Source/Runtime/Core/Private/ProfilingDebugging/MallocProfiler.cpp
+++ b/Engine/Source/Runtime/Core/Private/ProfilingDebugging/MallocProfiler.cpp
@@ -28,8 +28,13 @@ CORE_API FMallocProfiler* GMallocProfiler;
 
 /** Magic value, determining that file is a memory profiler file.				*/
 #define MEMORY_PROFILER_MAGIC						0xDA15F7D8
-/** Version of memory profiler.													*/
-#define MEMORY_PROFILER_VERSION						4
+/** 
+ * Version of memory profiler
+ * 
+ *	5 - Changed uint32 to uint64 in the FProfilerHeader to support profiler files larger than 4.0GB
+ *		Removed NumDataFiles, no longer used.
+ */
+#define MEMORY_PROFILER_VERSION						5
 
 /*=============================================================================
 	Profiler header.
@@ -49,25 +54,23 @@ struct FProfilerHeader
 	FString ExecutableName;
 
 	/** Offset in file for name table.					*/
-	uint32	NameTableOffset;
+	uint64	NameTableOffset;
 	/** Number of name table entries.					*/
-	uint32	NameTableEntries;
+	uint64	NameTableEntries;
 
 	/** Offset in file for callstack address table.		*/
-	uint32	CallStackAddressTableOffset;
+	uint64	CallStackAddressTableOffset;
 	/** Number of callstack address entries.			*/
-	uint32	CallStackAddressTableEntries;
+	uint64	CallStackAddressTableEntries;
 
 	/** Offset in file for callstack table.				*/
-	uint32	CallStackTableOffset;
+	uint64	CallStackTableOffset;
 	/** Number of callstack entries.					*/
-	uint32	CallStackTableEntries;
+	uint64	CallStackTableEntries;
 	/** The file offset for module information.			*/
-	uint32	ModulesOffset;
+	uint64	ModulesOffset;
 	/** The number of module entries.					*/
-	uint32	ModuleEntries;
-	/** Number of data files total.						*/
-	uint32	NumDataFiles;
+	uint64	ModuleEntries;
 
 	/**
 	 * Serialization operator.
@@ -91,8 +94,7 @@ struct FProfilerHeader
 			<< Header.CallStackTableOffset
 			<< Header.CallStackTableEntries
 			<< Header.ModulesOffset
-			<< Header.ModuleEntries
-			<< Header.NumDataFiles;
+			<< Header.ModuleEntries;
 
 		check( Ar.IsSaving() );
 		Header.ExecutableName.SerializeAsANSICharArray(Ar,255);
@@ -283,10 +285,6 @@ FMallocProfiler::FMallocProfiler(FMalloc* InMalloc)
 ,	MemoryOperationCount( 0 )
 {
 	StartTime = FPlatformTime::Seconds();
-
-	// Add callbacks for garbage collection
-	FCoreDelegates::PreGarbageCollect.AddStatic( FMallocProfiler::SnapshotMemoryGCStart );
-	FCoreDelegates::PostGarbageCollect.AddStatic( FMallocProfiler::SnapshotMemoryGCEnd );
 }
 
 /**
@@ -508,10 +506,7 @@ void FMallocProfiler::EndProfiling()
 		Header.Version				= MEMORY_PROFILER_VERSION;
 		Header.PlatformName			= FPlatformProperties::PlatformName();
 		Header.bShouldSerializeSymbolInfo = SERIALIZE_SYMBOL_INFO ? 1 : 0;
-#if PLATFORM_WINDOWS
 		Header.ExecutableName		= FPlatformProcess::ExecutableName();
-#endif // PLATFORM_WINDOWS
-		Header.NumDataFiles			= 1;
 
 		// Write out name table and update header with offset and count.
 		Header.NameTableOffset	= SymbolFileWriter->Tell();
diff --git a/Engine/Source/Runtime/Engine/Private/MallocProfilerEx.cpp b/Engine/Source/Runtime/Engine/Private/MallocProfilerEx.cpp
--- a/Engine/Source/Runtime/Engine/Private/MallocProfilerEx.cpp
+++ b/Engine/Source/Runtime/Engine/Private/MallocProfilerEx.cpp
@@ -22,7 +22,11 @@
  */
 FMallocProfilerEx::FMallocProfilerEx( FMalloc* InMalloc )
 	: FMallocProfiler( InMalloc )
-{}
+{
+	// Add callbacks for garbage collection
+	FCoreUObjectDelegates::PreGarbageCollect.AddStatic( FMallocProfiler::SnapshotMemoryGCStart );
+	FCoreUObjectDelegates::PostGarbageCollect.AddStatic( FMallocProfiler::SnapshotMemoryGCEnd );
+}
 
 /** 
  * Writes names of currently loaded levels. 
