Memory Profiler: Optimized FillActiveCallStackList to avoid frequent reallocations, significantly improving overall token stream parse speed

diff --git a/Engine/Source/Programs/MemoryProfiler2/StreamSnapshot.cs b/Engine/Source/Programs/MemoryProfiler2/StreamSnapshot.cs
--- a/Engine/Source/Programs/MemoryProfiler2/StreamSnapshot.cs
+++ b/Engine/Source/Programs/MemoryProfiler2/StreamSnapshot.cs
@@ -286,14 +286,7 @@ namespace MemoryProfiler2
 			}
 
 			// strip out any callstacks with no allocations
-			for( int i = ActiveCallStackList.Count - 1; i >= 0; i-- )
-			{
-				if( ActiveCallStackList[ i ].Count == 0 )
-				{
-					ActiveCallStackList.RemoveAt( i );
-				}
-			}
-
+			ActiveCallStackList.RemoveAll(Item => Item.Count == 0);
 			ActiveCallStackList.TrimExcess();
 		}
 
