Copying //UE4/Dev-Editor to //UE4/Dev-Main (Source: //UE4/Dev-Editor @ 3152045)

diff --git a/Engine/Source/Editor/UnrealEd/Private/Fbx/FbxFactory.cpp b/Engine/Source/Editor/UnrealEd/Private/Fbx/FbxFactory.cpp
--- a/Engine/Source/Editor/UnrealEd/Private/Fbx/FbxFactory.cpp
+++ b/Engine/Source/Editor/UnrealEd/Private/Fbx/FbxFactory.cpp
@@ -12,6 +12,7 @@
 #include "Engine/StaticMesh.h"
 #include "ObjectTools.h"
 #include "Animation/AnimSequence.h"
+#include "JsonObjectConverter.h"
 
 #define LOCTEXT_NAMESPACE "FBXFactory"
 
@@ -27,7 +28,7 @@ UFbxFactory::UFbxFactory(const FObjectInitializer& ObjectInitializer)
 	bText = false;
 	bEditorImport = true;
 	bOperationCanceled = false;
-	bDetectImportTypeOnImport = false;
+	bDetectImportTypeOnImport = true;
 }
 
 
@@ -190,9 +191,11 @@ UObject* UFbxFactory::FactoryCreateBinary
 	}
 
 
-	bool bShowImportDialog = bShowOption && !GIsAutomationTesting;
+	// Show the import dialog only when not in a "yes to all" state or when automating import
+	bool bIsAutomated = IsAutomatedImport();
+	bool bShowImportDialog = bShowOption && !bIsAutomated;
 	bool bImportAll = false;
-	ImportOptions = GetImportOptions(FbxImporter, ImportUI, bShowImportDialog, InParent->GetPathName(), bOperationCanceled, bImportAll, bIsObjFormat, bIsObjFormat, ForcedImportType );
+	ImportOptions = GetImportOptions(FbxImporter, ImportUI, bShowImportDialog, bIsAutomated, InParent->GetPathName(), bOperationCanceled, bImportAll, bIsObjFormat, bIsObjFormat, ForcedImportType );
 	bOutOperationCanceled = bOperationCanceled;
 	
 	if( bImportAll )
@@ -201,8 +204,12 @@ UObject* UFbxFactory::FactoryCreateBinary
 		bShowOption = false;
 	}
 
-	// For multiple files, use the same settings
-	bDetectImportTypeOnImport = false;
+	// Automated importing does not use the same settings and gets its settings straight from the user
+	if(!bIsAutomated)
+	{
+		// For multiple files, use the same settings
+		bDetectImportTypeOnImport = false;
+	}
 
 	if (ImportOptions)
 	{
@@ -651,6 +658,11 @@ bool UFbxFactory::FactoryCanImport(const FString& Filename)
 	return false;
 }
 
+IImportSettingsParser* UFbxFactory::GetImportSettingsParser()
+{
+	return ImportUI;
+}
+
 UFbxImportUI::UFbxImportUI(const FObjectInitializer& ObjectInitializer)
 	: Super(ObjectInitializer)
 {
@@ -688,4 +700,40 @@ bool UFbxImportUI::CanEditChange( const UProperty* InProperty ) const
 	return bIsMutable;
 }
 
+void UFbxImportUI::ParseFromJson(TSharedRef<class FJsonObject> ImportSettingsJson)
+{
+	// Skip instanced object references. 
+	int64 SkipFlags = CPF_InstancedReference;
+	FJsonObjectConverter::JsonObjectToUStruct(ImportSettingsJson, GetClass(), this, 0, SkipFlags);
+
+	const TSharedPtr<FJsonObject>* StaticMeshImportJson = nullptr;
+	ImportSettingsJson->TryGetObjectField(TEXT("StaticMeshImportData"), StaticMeshImportJson);
+	if(StaticMeshImportJson)
+	{
+		FJsonObjectConverter::JsonObjectToUStruct(StaticMeshImportJson->ToSharedRef(), StaticMeshImportData->GetClass(), StaticMeshImportData, 0, 0);
+	}
+
+	const TSharedPtr<FJsonObject>* SkeletalMeshImportJson = nullptr;
+	ImportSettingsJson->TryGetObjectField(TEXT("SkeletalMeshImportData"), SkeletalMeshImportJson);
+	if (SkeletalMeshImportJson)
+	{
+		FJsonObjectConverter::JsonObjectToUStruct(SkeletalMeshImportJson->ToSharedRef(), SkeletalMeshImportData->GetClass(), SkeletalMeshImportData, 0, 0);
+	}
+
+	const TSharedPtr<FJsonObject>* AnimImportJson = nullptr;
+	ImportSettingsJson->TryGetObjectField(TEXT("AnimSequenceImportData"), AnimImportJson);
+	if (AnimImportJson)
+	{
+		FJsonObjectConverter::JsonObjectToUStruct(AnimImportJson->ToSharedRef(), AnimSequenceImportData->GetClass(), AnimSequenceImportData, 0, 0);
+	}
+
+	const TSharedPtr<FJsonObject>* TextureImportJson = nullptr;
+	ImportSettingsJson->TryGetObjectField(TEXT("TextureImportData"), TextureImportJson);
+	if (TextureImportJson)
+	{
+		FJsonObjectConverter::JsonObjectToUStruct(TextureImportJson->ToSharedRef(), TextureImportData->GetClass(), TextureImportData, 0, 0);
+	}
+}
+
+
 #undef LOCTEXT_NAMESPACE
diff --git a/Engine/Source/Editor/UnrealEd/Private/Fbx/FbxMainImport.cpp b/Engine/Source/Editor/UnrealEd/Private/Fbx/FbxMainImport.cpp
--- a/Engine/Source/Editor/UnrealEd/Private/Fbx/FbxMainImport.cpp
+++ b/Engine/Source/Editor/UnrealEd/Private/Fbx/FbxMainImport.cpp
@@ -28,7 +28,7 @@ TSharedPtr<FFbxImporter> FFbxImporter::StaticInstance;
 
 
 
-FBXImportOptions* GetImportOptions( UnFbx::FFbxImporter* FbxImporter, UFbxImportUI* ImportUI, bool bShowOptionDialog, const FString& FullPath, bool& OutOperationCanceled, bool& bOutImportAll, bool bIsObjFormat, bool bForceImportType, EFBXImportType ImportType )
+FBXImportOptions* GetImportOptions( UnFbx::FFbxImporter* FbxImporter, UFbxImportUI* ImportUI, bool bShowOptionDialog, bool bIsAutomated, const FString& FullPath, bool& OutOperationCanceled, bool& bOutImportAll, bool bIsObjFormat, bool bForceImportType, EFBXImportType ImportType )
 {
 	OutOperationCanceled = false;
 
@@ -151,7 +151,7 @@ FBXImportOptions* GetImportOptions( UnFbx::FFbxImporter* FbxImporter, UFbxImport
 			OutOperationCanceled = true;
 		}
 	}
-	else if (GIsAutomationTesting)
+	else if (bIsAutomated)
 	{
 		//Automation tests set ImportUI settings directly.  Just copy them over
 		UnFbx::FBXImportOptions* ImportOptions = FbxImporter->GetImportOptions();
diff --git a/Engine/Source/Editor/UnrealEd/Private/Fbx/FbxMaterialImport.cpp b/Engine/Source/Editor/UnrealEd/Private/Fbx/FbxMaterialImport.cpp
--- a/Engine/Source/Editor/UnrealEd/Private/Fbx/FbxMaterialImport.cpp
+++ b/Engine/Source/Editor/UnrealEd/Private/Fbx/FbxMaterialImport.cpp
@@ -477,7 +477,7 @@ bool CanUseMaterialWithInstance(FbxSurfaceMaterial& FbxMaterial, const char* Mat
 }
 
 
-void UnFbx::FFbxImporter::CreateUnrealMaterial(FbxSurfaceMaterial& FbxMaterial, TArray<UMaterialInterface*>& OutMaterials, TArray<FString>& UVSets)
+void UnFbx::FFbxImporter::CreateUnrealMaterial(FbxSurfaceMaterial& FbxMaterial, TArray<UMaterialInterface*>& OutMaterials, TArray<FString>& UVSets, bool bForSkeletalMesh)
 {
 	// Make sure we have a parent
 	if ( !ensure(Parent.IsValid()) )
@@ -673,6 +673,12 @@ void UnFbx::FFbxImporter::CreateUnrealMaterial(FbxSurfaceMaterial& FbxMaterial,
 			// Notify the asset registry
 			FAssetRegistryModule::AssetCreated(UnrealMaterial);
 
+			if(bForSkeletalMesh)
+			{
+				bool bNeedsRecompile = false;
+				UnrealMaterial->GetMaterial()->SetMaterialUsage(bNeedsRecompile, MATUSAGE_SkeletalMesh);
+			}
+
 			// Set the dirty flag so this package will get saved later
 			Package->SetDirtyFlag(true);
 
@@ -717,7 +723,7 @@ void UnFbx::FFbxImporter::CreateUnrealMaterial(FbxSurfaceMaterial& FbxMaterial,
 	}
 }
 
-int32 UnFbx::FFbxImporter::CreateNodeMaterials(FbxNode* FbxNode, TArray<UMaterialInterface*>& OutMaterials, TArray<FString>& UVSets)
+int32 UnFbx::FFbxImporter::CreateNodeMaterials(FbxNode* FbxNode, TArray<UMaterialInterface*>& OutMaterials, TArray<FString>& UVSets, bool bForSkeletalMesh)
 {
 	int32 MaterialCount = FbxNode->GetMaterialCount();
 	for(int32 MaterialIndex=0; MaterialIndex < MaterialCount; ++MaterialIndex)
@@ -726,7 +732,7 @@ int32 UnFbx::FFbxImporter::CreateNodeMaterials(FbxNode* FbxNode, TArray<UMateria
 
 		if( FbxMaterial )
 		{
-			CreateUnrealMaterial(*FbxMaterial, OutMaterials, UVSets);
+			CreateUnrealMaterial(*FbxMaterial, OutMaterials, UVSets, bForSkeletalMesh);
 		}
 	}
 	return MaterialCount;
diff --git a/Engine/Source/Editor/UnrealEd/Private/Fbx/FbxSkeletalMeshImport.cpp b/Engine/Source/Editor/UnrealEd/Private/Fbx/FbxSkeletalMeshImport.cpp
--- a/Engine/Source/Editor/UnrealEd/Private/Fbx/FbxSkeletalMeshImport.cpp
+++ b/Engine/Source/Editor/UnrealEd/Private/Fbx/FbxSkeletalMeshImport.cpp
@@ -2240,7 +2240,8 @@ bool UnFbx::FFbxImporter::FillSkelMeshImporterFromFbx( FSkeletalMeshImportData&
 	TArray<UMaterialInterface*> Materials;
 	if (ImportOptions->bImportMaterials)
 	{
-		CreateNodeMaterials(Node, Materials, UVSets);
+		bool bForSkeletalMesh = true;
+		CreateNodeMaterials(Node, Materials, UVSets, bForSkeletalMesh);
 	}
 	else if (ImportOptions->bImportTextures)
 	{
diff --git a/Engine/Source/Editor/UnrealEd/Private/Fbx/FbxStaticMeshImport.cpp b/Engine/Source/Editor/UnrealEd/Private/Fbx/FbxStaticMeshImport.cpp
--- a/Engine/Source/Editor/UnrealEd/Private/Fbx/FbxStaticMeshImport.cpp
+++ b/Engine/Source/Editor/UnrealEd/Private/Fbx/FbxStaticMeshImport.cpp
@@ -271,7 +271,8 @@ bool UnFbx::FFbxImporter::BuildStaticMeshFromGeometry(FbxNode* Node, UStaticMesh
 	TArray<UMaterialInterface*> Materials;
 	if (ImportOptions->bImportMaterials)
 	{
-		CreateNodeMaterials(Node, Materials, FBXUVs.UVSets);
+		bool bForSkeletalMesh = false;
+		CreateNodeMaterials(Node, Materials, FBXUVs.UVSets, bForSkeletalMesh);
 	}
 	else if (ImportOptions->bImportTextures)
 	{
@@ -864,7 +865,7 @@ UStaticMesh* UnFbx::FFbxImporter::ReimportStaticMesh(UStaticMesh* Mesh, UFbxStat
 	{
 		Node = FbxMeshArray[0];
 	}
-	else
+	else if(!ImportOptions->bCombineToSingle)
 	{
 		// find the Fbx mesh node that the Unreal Mesh matches according to name
 		int32 MeshIndex;
diff --git a/Engine/Source/Editor/UnrealEd/Private/FbxImporter.h b/Engine/Source/Editor/UnrealEd/Private/FbxImporter.h
--- a/Engine/Source/Editor/UnrealEd/Private/FbxImporter.h
+++ b/Engine/Source/Editor/UnrealEd/Private/FbxImporter.h
@@ -407,7 +407,7 @@ public:
 	static FbxString	ConvertToFbxString(const FString& String);
 };
 
-FBXImportOptions* GetImportOptions( class FFbxImporter* FbxImporter, UFbxImportUI* ImportUI, bool bShowOptionDialog, const FString& FullPath, bool& OutOperationCanceled, bool& OutImportAll, bool bIsObjFormat, bool bForceImportType = false, EFBXImportType ImportType = FBXIT_StaticMesh );
+FBXImportOptions* GetImportOptions( class FFbxImporter* FbxImporter, UFbxImportUI* ImportUI, bool bShowOptionDialog, bool bIsAutomated, const FString& FullPath, bool& OutOperationCanceled, bool& OutImportAll, bool bIsObjFormat, bool bForceImportType = false, EFBXImportType ImportType = FBXIT_StaticMesh );
 void ApplyImportUIToImportOptions(UFbxImportUI* ImportUI, FBXImportOptions& InOutImportOptions);
 
 struct FImportedMaterialData
@@ -1148,7 +1148,7 @@ protected:
 	 * @param UVSets UV set name list
 	 * @return int32 material count that created from the Fbx node
 	 */
-	int32 CreateNodeMaterials(FbxNode* FbxNode, TArray<UMaterialInterface*>& outMaterials, TArray<FString>& UVSets);
+	int32 CreateNodeMaterials(FbxNode* FbxNode, TArray<UMaterialInterface*>& outMaterials, TArray<FString>& UVSets, bool bForSkeletalMesh);
 
 	/**
 	 * Make material Unreal asset name from the Fbx material
@@ -1167,7 +1167,7 @@ protected:
 	 * @param outMaterials Unreal Materials we created
 	 * @param outUVSets
 	 */
-	void CreateUnrealMaterial(FbxSurfaceMaterial& FbxMaterial, TArray<UMaterialInterface*>& OutMaterials, TArray<FString>& UVSets);
+	void CreateUnrealMaterial(FbxSurfaceMaterial& FbxMaterial, TArray<UMaterialInterface*>& OutMaterials, TArray<FString>& UVSets, bool bForSkeletalMesh);
 	
 	/**
 	 * Visit all materials of one node, import textures from materials.
diff --git a/Engine/Source/Editor/UnrealEd/Private/FileHelpers.cpp b/Engine/Source/Editor/UnrealEd/Private/FileHelpers.cpp
--- a/Engine/Source/Editor/UnrealEd/Private/FileHelpers.cpp
+++ b/Engine/Source/Editor/UnrealEd/Private/FileHelpers.cpp
@@ -3507,6 +3507,16 @@ void FEditorFileUtils::FindAllSubmittablePackageFiles(TMap<FString, FSourceContr
 
 	TArray<FString> Packages;
 	FEditorFileUtils::FindAllPackageFiles(Packages);
+
+	// Handle the project file
+	FSourceControlStatePtr ProjectFileSourceControlState = SourceControlProvider.GetState(FPaths::ConvertRelativePathToFull(FPaths::GetProjectFilePath()), EStateCacheUsage::Use);
+
+	if (ProjectFileSourceControlState.IsValid() && ProjectFileSourceControlState->IsCurrent() &&
+		(ProjectFileSourceControlState->IsCheckedOut() || ProjectFileSourceControlState->IsAdded() || (!ProjectFileSourceControlState->IsSourceControlled() && ProjectFileSourceControlState->CanAdd())))
+	{
+		OutPackages.Add(FPaths::GetProjectFilePath(), MoveTemp(ProjectFileSourceControlState));
+	}
+
 	for (TArray<FString>::TConstIterator PackageIter(Packages); PackageIter; ++PackageIter)
 	{
 		const FString Filename = *PackageIter;
diff --git a/Engine/Source/Editor/UnrealEd/Private/MaterialGraphSchema.cpp b/Engine/Source/Editor/UnrealEd/Private/MaterialGraphSchema.cpp
--- a/Engine/Source/Editor/UnrealEd/Private/MaterialGraphSchema.cpp
+++ b/Engine/Source/Editor/UnrealEd/Private/MaterialGraphSchema.cpp
@@ -373,12 +373,15 @@ void UMaterialGraphSchema::GetContextMenuActions(const UEdGraph* CurrentGraph, c
 			// Only display the 'Break Link' option if there is a link to break!
 			if (InGraphPin->LinkedTo.Num() > 0)
 			{
-				MenuBuilder->AddMenuEntry(
-					LOCTEXT("SelectLinkedNodes", "Select Linked Nodes"),
-					LOCTEXT("SelectLinkedNodesTooltip", "Adds all input Nodes linked to this Pin to selection"),
-					FSlateIcon(),
-					FUIAction(FExecuteAction::CreateUObject((UMaterialGraphSchema*const)this, &UMaterialGraphSchema::SelectAllInputNodes, const_cast<UEdGraph*>(CurrentGraph), const_cast<UEdGraphPin*>(InGraphPin)))
-					);
+				if(InGraphPin->Direction == EEdGraphPinDirection::EGPD_Input)
+				{
+					MenuBuilder->AddMenuEntry(
+						LOCTEXT("SelectLinkedNodes", "Select Linked Nodes"),
+						LOCTEXT("SelectLinkedNodesTooltip", "Adds all input Nodes linked to this Pin to selection"),
+						FSlateIcon(),
+						FUIAction(FExecuteAction::CreateUObject((UMaterialGraphSchema*const)this, &UMaterialGraphSchema::SelectAllInputNodes, const_cast<UEdGraph*>(CurrentGraph), const_cast<UEdGraphPin*>(InGraphPin)))
+						);
+				}
 
 				MenuBuilder->AddMenuEntry(FGraphEditorCommands::Get().BreakPinLinks);
 
diff --git a/Engine/Source/Editor/UnrealEd/Private/PerformanceMonitor.cpp b/Engine/Source/Editor/UnrealEd/Private/PerformanceMonitor.cpp
--- a/Engine/Source/Editor/UnrealEd/Private/PerformanceMonitor.cpp
+++ b/Engine/Source/Editor/UnrealEd/Private/PerformanceMonitor.cpp
@@ -183,7 +183,7 @@ void FPerformanceMonitor::AutoApplyScalability()
 	Scalability::SaveState(GEditorSettingsIni);
 	GEditor->RedrawAllViewports();
 
-	const bool bAutoApplied = true;
+	const bool bAutoApplied = false;
 	Scalability::RecordQualityLevelsAnalytics(bAutoApplied);
 
 	GEditor->DisableRealtimeViewports();
@@ -321,14 +321,14 @@ void FPerformanceMonitor::Tick(float DeltaTime)
 		// Choose an appropriate message
 		enum MessagesEnum { Seconds, SecondsPercent, Minute, MinutePercent, Minutes, MinutesPercent };
 		const FText Messages[] = {
-			LOCTEXT("PerformanceWarningInProgress_Seconds",			"Your framerate has been under {Framerate} FPS for the past {SampleTime} seconds.\n\nApplying reduced quality settings in {TimeRemaining}s."),
-			LOCTEXT("PerformanceWarningInProgress_Seconds_Percent", "Your framerate has been under {Framerate} FPS for {Percentage}% of the past {SampleTime} seconds.\n\nApplying reduced quality settings in {TimeRemaining}s."),
+			LOCTEXT("PerformanceWarningInProgress_Seconds",			"Your framerate has been under {Framerate} FPS for the past {SampleTime} seconds.\n\nDo you want to apply reduced quality settings? {TimeRemaining}s"),
+			LOCTEXT("PerformanceWarningInProgress_Seconds_Percent", "Your framerate has been under {Framerate} FPS for {Percentage}% of the past {SampleTime} seconds.\n\nDo you want to apply reduced quality settings? {TimeRemaining}s"),
 
-			LOCTEXT("PerformanceWarningInProgress_Minute",			"Your framerate has been under {Framerate} FPS for the past minute.\n\nApplying reduced quality settings in {TimeRemaining}s."),
-			LOCTEXT("PerformanceWarningInProgress_Minute_Percent",	"Your framerate has been under {Framerate} FPS for {Percentage}% of the last minute.\n\nApplying reduced quality settings in {TimeRemaining}s."),
+			LOCTEXT("PerformanceWarningInProgress_Minute",			"Your framerate has been under {Framerate} FPS for the past minute.\n\nDo you want to apply reduced quality settings? {TimeRemaining}s"),
+			LOCTEXT("PerformanceWarningInProgress_Minute_Percent",	"Your framerate has been under {Framerate} FPS for {Percentage}% of the last minute.\n\nDo you want to apply reduced quality settings? {TimeRemaining}s"),
 
-			LOCTEXT("PerformanceWarningInProgress_Minutes",			"Your framerate has been below {Framerate} FPS for the past {SampleTime} minutes.\n\nApplying reduced quality settings in {TimeRemaining}s."),
-			LOCTEXT("PerformanceWarningInProgress_Minutes_Percent", "Your framerate has been below {Framerate} FPS for {Percentage}% of the past {SampleTime} minutes.\n\nApplying reduced quality settings in {TimeRemaining}s."),
+			LOCTEXT("PerformanceWarningInProgress_Minutes",			"Your framerate has been below {Framerate} FPS for the past {SampleTime} minutes.\n\nDo you want to apply reduced quality settings? {TimeRemaining}s"),
+			LOCTEXT("PerformanceWarningInProgress_Minutes_Percent", "Your framerate has been below {Framerate} FPS for {Percentage}% of the past {SampleTime} minutes.\n\nDo you want to apply reduced quality settings? {TimeRemaining}s"),
 		};
 
 		int32 Message;
@@ -363,8 +363,8 @@ void FPerformanceMonitor::Tick(float DeltaTime)
 			if (NotificationTimeout <= 0)
 			{
 				// Timed-out. Apply the settings.
-				AutoApplyScalability();
 				Reset();
+				bIsNotificationAllowed = false;
 			}
 			else
 			{
diff --git a/Engine/Source/Editor/UnrealEd/Private/PlayLevel.cpp b/Engine/Source/Editor/UnrealEd/Private/PlayLevel.cpp
--- a/Engine/Source/Editor/UnrealEd/Private/PlayLevel.cpp
+++ b/Engine/Source/Editor/UnrealEd/Private/PlayLevel.cpp
@@ -869,9 +869,14 @@ void UEditorEngine::PlaySessionSingleStepped()
 	FEditorDelegates::SingleStepPIE.Broadcast(bIsSimulatingInEditor);
 }
 
-bool UEditorEngine::ProcessDebuggerCommands(const FKey InKey, const FModifierKeysState ModifierKeyState)
+bool UEditorEngine::ProcessDebuggerCommands(const FKey InKey, const FModifierKeysState ModifierKeyState, EInputEvent EventType )
 {
-	return FPlayWorldCommands::GlobalPlayWorldActions->ProcessCommandBindings(InKey, ModifierKeyState, false);
+	if( EventType == IE_Pressed )
+	{
+		return FPlayWorldCommands::GlobalPlayWorldActions->ProcessCommandBindings(InKey, ModifierKeyState, false);
+	}
+	
+	return false;
 }
 
 /* fits the window position to make sure it falls within the confines of the desktop */
diff --git a/Engine/Source/Editor/UnrealEd/Private/ReferenceInfoUtils.cpp b/Engine/Source/Editor/UnrealEd/Private/ReferenceInfoUtils.cpp
--- a/Engine/Source/Editor/UnrealEd/Private/ReferenceInfoUtils.cpp
+++ b/Engine/Source/Editor/UnrealEd/Private/ReferenceInfoUtils.cpp
@@ -460,11 +460,11 @@ namespace ReferenceInfoUtils
 		}
 
 		// Create a string for the resource size.
-		const SIZE_T ResourceSize = ReferencedObject->GetResourceSize(EResourceSizeMode::Exclusive);
+		const SIZE_T ReferencedObjectResourceSize = ReferencedObject->GetResourceSizeBytes(EResourceSizeMode::Exclusive);
 		FString ResourceSizeString;
-		if ( ResourceSize > 0 )
+		if ( ReferencedObjectResourceSize > 0 )
 		{
-			ResourceSizeString = FString::Printf( TEXT("%.2f"), ((float)ResourceSize)/1024.f );
+			ResourceSizeString = FString::Printf( TEXT("%.2f"), ((float)ReferencedObjectResourceSize)/1024.f );
 		}
 
 		FString ObjectPathName;
diff --git a/Engine/Source/Editor/UnrealEd/Private/SCurveEditor.cpp b/Engine/Source/Editor/UnrealEd/Private/SCurveEditor.cpp
--- a/Engine/Source/Editor/UnrealEd/Private/SCurveEditor.cpp
+++ b/Engine/Source/Editor/UnrealEd/Private/SCurveEditor.cpp
@@ -10,6 +10,7 @@
 #include "SNumericEntryBox.h"
 #include "STextEntryPopup.h"
 #include "CurveEditorSettings.h"
+#include "Factories/CurveFactory.h"
 
 #define LOCTEXT_NAMESPACE "SCurveEditor"
 
diff --git a/Engine/Source/Editor/UnrealEd/Private/Tests/AutomationEditorCommon.cpp b/Engine/Source/Editor/UnrealEd/Private/Tests/AutomationEditorCommon.cpp
--- a/Engine/Source/Editor/UnrealEd/Private/Tests/AutomationEditorCommon.cpp
+++ b/Engine/Source/Editor/UnrealEd/Private/Tests/AutomationEditorCommon.cpp
@@ -761,7 +761,7 @@ bool FOpenEditorForAssetCommand::Update()
 		if ( FAssetEditorManager::Get().FindEditorForAsset(Object, true) != NULL )
 		{
 			UE_LOG(LogEditorAutomationTests, Log, TEXT("Verified asset editor for: %s."), *AssetName);
-			UE_LOG(LogEditorAutomationTests, Display, TEXT("The editor successffully loaded for: %s."), *AssetName);
+			UE_LOG(LogEditorAutomationTests, Display, TEXT("The editor successfully loaded for: %s."), *AssetName);
 			return true;
 		}
 	}
@@ -787,7 +787,7 @@ bool FCloseAllAssetEditorsCommand::Update()
 	}
 
 	UE_LOG(LogEditorAutomationTests, Log, TEXT("Verified asset editors were closed"));
-	UE_LOG(LogEditorAutomationTests, Display, TEXT("The asset editors closed successffully"));
+	UE_LOG(LogEditorAutomationTests, Display, TEXT("The asset editors closed successfully"));
 	return true;
 }
 
diff --git a/Engine/Source/Editor/UnrealEd/Public/AutomatedAssetImportData.h b/Engine/Source/Editor/UnrealEd/Public/AutomatedAssetImportData.h
new file mode 100644
--- /dev/null
+++ b/Engine/Source/Editor/UnrealEd/Public/AutomatedAssetImportData.h
@@ -0,0 +1,57 @@
+// Copyright 1998-2016 Epic Games, Inc. All Rights Reserved.
+
+#pragma once
+
+#include "AutomatedAssetImportData.generated.h"
+
+DECLARE_LOG_CATEGORY_EXTERN(LogAutomatedImport, Log, All);
+
+/**
+ * Contains data for a group of assets to import
+ */ 
+UCLASS(Transient)
+class UNREALED_API UAutomatedAssetImportData : public UObject
+{
+	GENERATED_BODY()
+
+public:
+	UAutomatedAssetImportData();
+
+	/** @return true if this group contains enough valid data to import*/
+	bool IsValid() const;
+
+	/** Initalizes the group */
+	void Initialize();
+
+	/** @return the display name of the group */
+	FString GetDisplayName() const; 
+public:
+	/** Display name of the group. This is for logging purposes only. */
+	UPROPERTY()
+	FString GroupName;
+
+	/** Filenames to import */
+	UPROPERTY()
+	TArray<FString> Filenames;
+
+	/** Content path in the projects content directory where assets will be imported */
+	UPROPERTY()
+	FString DestinationPath;
+
+	/** Name of the factory to use when importing these assets. If not specified the factory type will be auto detected */
+	UPROPERTY()
+	FString FactoryName;
+
+	/** Whether or not to replace existing assets */
+	UPROPERTY()
+	bool bReplaceExisting;
+
+	/** Whether or not to skip importing over read only assets that could not be checked out */
+	UPROPERTY()
+	bool bSkipReadOnly;
+
+	/** Pointer to the factory currently being sued */
+	UPROPERTY()
+	UFactory* Factory;
+
+};
diff --git a/Engine/Source/Editor/UnrealEd/Public/SCSVImportOptions.h b/Engine/Source/Editor/UnrealEd/Public/SCSVImportOptions.h
--- a/Engine/Source/Editor/UnrealEd/Public/SCSVImportOptions.h
+++ b/Engine/Source/Editor/UnrealEd/Public/SCSVImportOptions.h
@@ -5,6 +5,8 @@
  */
 
 #pragma once
+#include "Factories/CSVImportFactory.h"
+
 
 class UNREALED_API SCSVImportOptions : public SCompoundWidget
 {
diff --git a/Engine/Source/Editor/UnrealEd/Public/SCurveEditor.h b/Engine/Source/Editor/UnrealEd/Public/SCurveEditor.h
--- a/Engine/Source/Editor/UnrealEd/Public/SCurveEditor.h
+++ b/Engine/Source/Editor/UnrealEd/Public/SCurveEditor.h
@@ -7,6 +7,8 @@
 #include "CurveEditorSettings.h"
 #include "EditorUndoClient.h"
 
+class UCurveFactory;
+
 //////////////////////////////////////////////////////////////////////////
 // FTrackScaleInfo
 
diff --git a/Engine/Source/Editor/UnrealEd/Public/UnrealEd.h b/Engine/Source/Editor/UnrealEd/Public/UnrealEd.h
--- a/Engine/Source/Editor/UnrealEd/Public/UnrealEd.h
+++ b/Engine/Source/Editor/UnrealEd/Public/UnrealEd.h
@@ -126,14 +126,6 @@
 #include "Classes/Factories/BlueprintMacroFactory.h"
 #include "Classes/Factories/BlueprintInterfaceFactory.h"
 #include "Classes/Factories/CameraAnimFactory.h"
-#include "Classes/Factories/CSVImportFactory.h"
-#include "Classes/Factories/ReimportCurveFactory.h"
-#include "Classes/Factories/ReimportCurveTableFactory.h"
-#include "Classes/Factories/ReimportDataTableFactory.h"
-#include "Classes/Factories/CurveFactory.h"
-#include "Classes/Factories/CurveImportFactory.h"
-#include "Classes/Factories/DataAssetFactory.h"
-#include "Classes/Factories/DataTableFactory.h"
 #include "Classes/Factories/DestructibleMeshFactory.h"
 #include "Classes/Factories/ReimportDestructibleMeshFactory.h"
 #include "Classes/Factories/DialogueVoiceFactory.h"
diff --git a/Engine/Source/Editor/UnrealEd/UnrealEd.Build.cs b/Engine/Source/Editor/UnrealEd/UnrealEd.Build.cs
--- a/Engine/Source/Editor/UnrealEd/UnrealEd.Build.cs
+++ b/Engine/Source/Editor/UnrealEd/UnrealEd.Build.cs
@@ -221,8 +221,7 @@ public class UnrealEd : ModuleRules
                 "PortalProxies",
                 "PortalServices",
                 "GeometryCacheEd",
-                "BlueprintNativeCodeGen",
-                "EditorAutomation",
+                "BlueprintNativeCodeGen"
 			}
 		);
 
diff --git a/Engine/Source/Programs/AutomationTool/AutomationUtils/ProjectParams.cs b/Engine/Source/Programs/AutomationTool/AutomationUtils/ProjectParams.cs
--- a/Engine/Source/Programs/AutomationTool/AutomationUtils/ProjectParams.cs
+++ b/Engine/Source/Programs/AutomationTool/AutomationUtils/ProjectParams.cs
@@ -307,7 +307,6 @@ namespace AutomationTool
             this.bUsesSteam = InParams.bUsesSteam;
 			this.bUsesCEF3 = InParams.bUsesCEF3;
 			this.bUsesSlate = InParams.bUsesSlate;
-			this.bUsesSlateEditorStyle = InParams.bUsesSlateEditorStyle;
             this.bDebugBuildsActuallyUseDebugCRT = InParams.bDebugBuildsActuallyUseDebugCRT;
 			this.Archive = InParams.Archive;
 			this.ArchiveDirectoryParam = InParams.ArchiveDirectoryParam;
@@ -1366,12 +1365,6 @@ namespace AutomationTool
 		/// </summary>
 		public bool bUsesSlate = true;
 
-		/// <summary>
-		/// Hack for legacy game styling isses.  No new project should ever set this to true
-		/// Whether the project uses the Slate editor style in game.  
-		/// </summary>
-		public bool bUsesSlateEditorStyle = false;
-
         /// <summary>
         /// By default we use the Release C++ Runtime (CRT), even when compiling Debug builds.  This is because the Debug C++
         /// Runtime isn't very useful when debugging Unreal Engine projects, and linking against the Debug CRT libraries forces
@@ -1648,7 +1641,6 @@ namespace AutomationTool
 			bUsesSteam = Properties.bUsesSteam;
 			bUsesCEF3 = Properties.bUsesCEF3;
 			bUsesSlate = Properties.bUsesSlate;
-			bUsesSlateEditorStyle = Properties.bUsesSlateEditorStyle;
             bDebugBuildsActuallyUseDebugCRT = Properties.bDebugBuildsActuallyUseDebugCRT;
 
 			bIsCodeBasedProject = Properties.bIsCodeBasedProject;			
@@ -1706,7 +1698,6 @@ namespace AutomationTool
 						GameTarget = TargetData.TargetName;
                         bDebugBuildsActuallyUseDebugCRT = TargetData.Rules.bDebugBuildsActuallyUseDebugCRT;
 						bUsesSlate = TargetData.Rules.bUsesSlate;
-						bUsesSlateEditorStyle = TargetData.Rules.bUsesSlateEditorStyle;
 						bUsesSteam = TargetData.Rules.bUsesSteam;
 						bUsesCEF3 = TargetData.Rules.bUsesCEF3;
 						ProjectType = ValidTarget;
@@ -1733,7 +1724,6 @@ namespace AutomationTool
 
 				bDebugBuildsActuallyUseDebugCRT = TargetData.Rules.bDebugBuildsActuallyUseDebugCRT;
 				bUsesSlate = TargetData.Rules.bUsesSlate;
-				bUsesSlateEditorStyle = TargetData.Rules.bUsesSlateEditorStyle;
 				bUsesSteam = TargetData.Rules.bUsesSteam;
 				bUsesCEF3 = TargetData.Rules.bUsesCEF3;
 				ProjectType = TargetRules.TargetType.Program;
diff --git a/Engine/Source/Programs/AutomationTool/AutomationUtils/ProjectUtils.cs b/Engine/Source/Programs/AutomationTool/AutomationUtils/ProjectUtils.cs
--- a/Engine/Source/Programs/AutomationTool/AutomationUtils/ProjectUtils.cs
+++ b/Engine/Source/Programs/AutomationTool/AutomationUtils/ProjectUtils.cs
@@ -544,7 +544,6 @@ namespace AutomationTool
 					Properties.bUsesCEF3 |= Rules.bUsesCEF3;
 					Properties.bUsesSlate |= Rules.bUsesSlate;
                     Properties.bDebugBuildsActuallyUseDebugCRT |= Rules.bDebugBuildsActuallyUseDebugCRT;
-					Properties.bUsesSlateEditorStyle |= Rules.bUsesSlateEditorStyle;
 				}
 			}
 		}
diff --git a/Engine/Source/Programs/MemoryProfiler2/AllocationTags.cs b/Engine/Source/Programs/MemoryProfiler2/AllocationTags.cs
new file mode 100644
--- /dev/null
+++ b/Engine/Source/Programs/MemoryProfiler2/AllocationTags.cs
@@ -0,0 +1,253 @@
+﻿using System;
+using System.Collections.Generic;
+using System.Linq;
+using System.Text;
+
+namespace MemoryProfiler2
+{
+	public class FAllocationTagHierarchy
+	{
+		public void InitializeHierarchy()
+		{
+			RootNode = new FNode(-1, null, null);
+
+			var Untagged = new FAllocationTag("Untagged");
+			UntaggedNode = new FNode(Untagged.TagNameIndex, RootNode, Untagged);
+		}
+
+		public void FinalizeHierarchy()
+		{
+			RootNode.SortChildrenRecursive();
+		}
+
+		public void RegisterTag(string TagStr, FAllocationTag Tag)
+		{
+			string CategoryName = "Uncategorized";
+			string ActualTag = TagStr;
+
+			// Split the tag into its category and remaining parts
+			int CategorySeparatorIndex = TagStr.IndexOf(':');
+			if (CategorySeparatorIndex != -1)
+			{
+				CategoryName = TagStr.Substring(0, CategorySeparatorIndex);
+				ActualTag = TagStr.Substring(CategorySeparatorIndex + 1);
+			}
+
+			// Add the top-level category
+			FNode CurrentNode = RootNode;
+			{
+				int CategoryNameIndex = FStreamInfo.GlobalInstance.GetNameIndex(CategoryName, true);
+				CurrentNode = CurrentNode.EnsureChildNode(CategoryNameIndex);
+			}
+
+			// Split the remaining parts it hierarchical parts, add a name entry for each tag, and add to the hierarchy
+			var HierarchicalTags = ActualTag.Split(TagPathSeparators, StringSplitOptions.RemoveEmptyEntries);
+			var HierarchicalTagNameIndices = new int[HierarchicalTags.Length];
+			for (int HierarchicalTagIndex = 0; HierarchicalTagIndex < HierarchicalTags.Length; ++HierarchicalTagIndex)
+			{
+				int HierarchicalTagNameIndex = FStreamInfo.GlobalInstance.GetNameIndex(HierarchicalTags[HierarchicalTagIndex], true);
+
+				// Are we dealing with a child node, or a tag?
+				if (HierarchicalTagIndex + 1 == HierarchicalTags.Length)
+				{
+					CurrentNode.EnsureChildTag(HierarchicalTagNameIndex, Tag);
+				}
+				else
+				{
+					CurrentNode = CurrentNode.EnsureChildNode(HierarchicalTagNameIndex);
+				}
+			}
+		}
+
+		public INode GetRootNode()
+		{
+			return RootNode;
+		}
+
+		public INode GetUntaggedNode()
+		{
+			return UntaggedNode;
+		}
+
+		private static char[] TagPathSeparators = { '/' };
+
+		public interface INode
+		{
+			string GetNodeName();
+
+			IEnumerable<INode> GetChildNodes();
+
+			IEnumerable<INode> GetChildTags();
+
+			FAllocationTag? GetTag();
+		}
+
+		private class FNode : INode
+		{
+			public FNode(int InNameIndex, FNode InParent, FAllocationTag? InTag)
+			{
+				Parent = InParent;
+				if (Parent != null)
+				{
+					if (InTag.HasValue)
+					{
+						Parent.RegisterChildTag(InNameIndex, this);
+					}
+					else
+					{
+						Parent.RegisterChildNode(InNameIndex, this);
+					}
+				}
+
+				NameIndex = InNameIndex;
+				Tag = InTag;
+
+				ChildNodes = new List<FNode>();
+				ChildNodesByNameMap = new Dictionary<int, FNode>();
+
+				ChildTags = new List<FNode>();
+				ChildTagsByNameMap = new Dictionary<int, FNode>();
+			}
+
+			public string GetNodeName()
+			{
+				return FStreamInfo.GlobalInstance.NameArray[NameIndex];
+			}
+
+			public IEnumerable<INode> GetChildNodes()
+			{
+				return ChildNodes;
+			}
+
+			public IEnumerable<INode> GetChildTags()
+			{
+				return ChildTags;
+			}
+
+			public FAllocationTag? GetTag()
+			{
+				return Tag;
+			}
+
+			public FNode EnsureChildNode(int InChildNameIndex)
+			{
+				FNode ChildNode;
+				if (!ChildNodesByNameMap.TryGetValue(InChildNameIndex, out ChildNode))
+				{
+					ChildNode = new FNode(InChildNameIndex, this, null);
+				}
+				return ChildNode;
+			}
+
+			public void EnsureChildTag(int InChildNameIndex, FAllocationTag InTag)
+			{
+				if (!ChildTagsByNameMap.ContainsKey(InChildNameIndex))
+				{
+					new FNode(InChildNameIndex, this, InTag);
+				}
+			}
+
+			public void SortChildrenRecursive()
+			{
+				// We sort by display name
+				ChildNodes.Sort((First, Second) => FStreamInfo.GlobalInstance.NameArray[First.NameIndex].CompareTo(FStreamInfo.GlobalInstance.NameArray[Second.NameIndex]));
+				ChildTags.Sort((First, Second) => FStreamInfo.GlobalInstance.NameArray[First.NameIndex].CompareTo(FStreamInfo.GlobalInstance.NameArray[Second.NameIndex]));
+
+				foreach (var ChildNode in ChildNodes)
+				{
+					ChildNode.SortChildrenRecursive();
+				}
+			}
+
+			void RegisterChildNode(int InChildNameIndex, FNode InChild)
+			{
+				ChildNodes.Add(InChild);
+				ChildNodesByNameMap.Add(InChildNameIndex, InChild);
+			}
+
+			void RegisterChildTag(int InChildNameIndex, FNode InChild)
+			{
+				ChildTags.Add(InChild);
+				ChildTagsByNameMap.Add(InChildNameIndex, InChild);
+			}
+
+			FNode Parent;
+
+			int NameIndex;
+			FAllocationTag? Tag;
+
+			List<FNode> ChildNodes;
+			Dictionary<int, FNode> ChildNodesByNameMap;
+
+			List<FNode> ChildTags;
+			Dictionary<int, FNode> ChildTagsByNameMap;
+		};
+
+		FNode RootNode;
+
+		FNode UntaggedNode;
+	};
+
+	public struct FAllocationTag : IEquatable<FAllocationTag>
+	{
+		public FAllocationTag(string TagStr)
+		{
+			// Add a name entry for this whole tag
+			TagNameIndex = FStreamInfo.GlobalInstance.GetNameIndex(TagStr, true);
+		}
+
+		public string Name
+		{
+			get
+			{
+				return FStreamInfo.GlobalInstance.NameArray[TagNameIndex];
+			}
+		}
+
+		public override int GetHashCode()
+		{
+			return TagNameIndex;
+		}
+
+		public override bool Equals(object Other)
+		{
+			return Other is FAllocationTag && ((FAllocationTag)Other).TagNameIndex == TagNameIndex;
+		}
+
+		public bool Equals(FAllocationTag Other)
+		{
+			return Other.TagNameIndex == TagNameIndex;
+		}
+
+		public static bool operator==(FAllocationTag LHS, FAllocationTag RHS)
+		{
+			return LHS.TagNameIndex == RHS.TagNameIndex;
+		}
+
+		public static bool operator!=(FAllocationTag LHS, FAllocationTag RHS)
+		{
+			return LHS.TagNameIndex != RHS.TagNameIndex;
+		}
+
+		public readonly int TagNameIndex;
+	};
+
+	public class FAllocationTags
+	{
+		public FAllocationTags(string TagsStr)
+		{
+			var SplitTags = TagsStr.Split(TagSeparators, StringSplitOptions.RemoveEmptyEntries);
+
+			Tags = new FAllocationTag[SplitTags.Length];
+			for (int TagIndex = 0; TagIndex < SplitTags.Length; ++TagIndex)
+			{
+				Tags[TagIndex] = new FAllocationTag(SplitTags[TagIndex]);
+				FStreamInfo.GlobalInstance.TagHierarchy.RegisterTag(SplitTags[TagIndex], Tags[TagIndex]);
+			}
+		}
+
+		public readonly FAllocationTag[] Tags;
+
+		private static char[] TagSeparators = { ';' };
+	}
+}
diff --git a/Engine/Source/Programs/MemoryProfiler2/CallGraphTreeViewParser.cs b/Engine/Source/Programs/MemoryProfiler2/CallGraphTreeViewParser.cs
--- a/Engine/Source/Programs/MemoryProfiler2/CallGraphTreeViewParser.cs
+++ b/Engine/Source/Programs/MemoryProfiler2/CallGraphTreeViewParser.cs
@@ -44,10 +44,22 @@ namespace MemoryProfiler2
 			var TruncatedNode = new FCallGraphNode("Truncated Callstacks");
 			var RegularNode = new FCallGraphNode("Full Callstacks");
 
+			bool bFilterIn = OwnerWindow.IsFilteringIn();
+
 			using ( FScopedLogTimer ParseTiming = new FScopedLogTimer( "FCallGraphTreeViewParser.ParseSnapshot" ) )
 			{
+				var FilteredCallstackList = new List<FCallStackAllocationInfo>(CallStackList.Count);
+				foreach (var AllocationInfo in CallStackList)
+				{
+					var FilteredAllocationInfo = AllocationInfo.GetAllocationInfoForTags(OwnerWindow.GetTagsFilter(), bFilterIn);
+					if (FilteredAllocationInfo.TotalCount != 0)
+					{
+						FilteredCallstackList.Add(FilteredAllocationInfo);
+					}
+				}
+
 				// Iterate over all call graph paths and add them to the graph.
-				foreach( FCallStackAllocationInfo AllocationInfo in CallStackList )
+				foreach ( FCallStackAllocationInfo AllocationInfo in FilteredCallstackList )
 				{
 					// Update progress bar.
 					if( CallStackCurrent >= NextProgressUpdate )
@@ -62,10 +74,8 @@ namespace MemoryProfiler2
 					FCallStack CallStack = FStreamInfo.GlobalInstance.CallStackArray[ AllocationInfo.CallStackIndex ];
 					// Split the tree into full and truncated callstacks.
 					var RootNode = CallStack.bIsTruncated ? TruncatedNode : RegularNode;
-					// Don't bother with callstacks that don't have a contribution.
-					if( ( ( AllocationInfo.Count != 0 ) || ( AllocationInfo.Size != 0 ) )
-						// Apply filter based on text representation of address.
-						&& CallStack.RunFilters( FilterText, OwnerWindow.Options.ClassGroups, OwnerWindow.IsFilteringIn(), OwnerWindow.SelectedMemoryPool ) )
+					// Apply filter based on text representation of address.
+					if( CallStack.RunFilters( FilterText, OwnerWindow.Options.ClassGroups, bFilterIn, OwnerWindow.SelectedMemoryPool ) )
 					{
 						// Add call stack to proper part of graph.
 						AddCallStackToGraph( RootNode, CallStack, AllocationInfo, ParentFunctionIndex, bInvertCallStacks );
@@ -133,14 +143,17 @@ namespace MemoryProfiler2
 			int FunctionIndex = FStreamInfo.GlobalInstance.CallStackAddressArray[AddressIndex].FunctionIndex;
 			FCallStack CallStack = FStreamInfo.GlobalInstance.CallStackArray[AllocationInfo.CallStackIndex];
 
+			var TotalAllocationSize = AllocationInfo.TotalSize;
+			var TotalAllocationCount = AllocationInfo.TotalCount;
+
 			// Iterate over existing nodes to see whether there is a match.
 			foreach (FCallGraphNode Node in ParentNode.Children)
 			{
 				// If there is a match, update the allocation size and return the current node.
 				if (FStreamInfo.GlobalInstance.CallStackAddressArray[Node.AddressIndex].FunctionIndex == FunctionIndex)
 				{
-					Node.AllocationSize += AllocationInfo.Size;
-					Node.AllocationCount += AllocationInfo.Count;
+					Node.AllocationSize += TotalAllocationSize;
+					Node.AllocationCount += TotalAllocationCount;
 					Node.CallStacks.Add(CallStack);
 
 					// Return current node as parent for next iteration.
@@ -150,7 +163,7 @@ namespace MemoryProfiler2
 
 			// If we made it here it means that we need to add a new node.
 			string FunctionName = FStreamInfo.GlobalInstance.NameArray[FunctionIndex];
-			FCallGraphNode NewNode = new FCallGraphNode(ParentNode, AddressIndex, AllocationInfo.Size, AllocationInfo.Count, FunctionName, CallStack);
+			FCallGraphNode NewNode = new FCallGraphNode(ParentNode, AddressIndex, TotalAllocationSize, TotalAllocationCount, FunctionName, CallStack);
 			return NewNode;
 		}
 
diff --git a/Engine/Source/Programs/MemoryProfiler2/ExclusiveListViewParser.cs b/Engine/Source/Programs/MemoryProfiler2/ExclusiveListViewParser.cs
--- a/Engine/Source/Programs/MemoryProfiler2/ExclusiveListViewParser.cs
+++ b/Engine/Source/Programs/MemoryProfiler2/ExclusiveListViewParser.cs
@@ -60,36 +60,46 @@ namespace MemoryProfiler2
 
 			ExclusiveListView.ListViewItemSorter = null; // clear this to avoid a Sort for each call to Add
 
-			// Sort based on passed in metric.
-			if( bShouldSortBySize )
-			{
-                CallStackList.Sort( CompareAbsSize );
-			}
-			else
-			{
-				CallStackList.Sort( CompareCount );
-			}
-
 			bool bFilterIn = OwnerWindow.IsFilteringIn();
 
 			using( FScopedLogTimer ParseTiming = new FScopedLogTimer( "FExclusiveListViewParser.ParseSnapshot" ) )
 			{
+				var FilteredCallstackList = new List<FCallStackAllocationInfo>(CallStackList.Count);
+				foreach (var AllocationInfo in CallStackList)
+				{
+					var FilteredAllocationInfo = AllocationInfo.GetAllocationInfoForTags(OwnerWindow.GetTagsFilter(), bFilterIn);
+					if (FilteredAllocationInfo.TotalCount != 0)
+					{
+						FilteredCallstackList.Add(FilteredAllocationInfo);
+					}
+				}
+
+				// Sort based on passed in metric.
+				if (bShouldSortBySize)
+				{
+					FilteredCallstackList.Sort(CompareAbsSize);
+				}
+				else
+				{
+					FilteredCallstackList.Sort(CompareCount);
+				}
+
 				// Figure out total size and count for percentages.
 				long TotalSize = 0;
 				long TotalCount = 0;
-				foreach( FCallStackAllocationInfo AllocationInfo in CallStackList )
+				foreach( FCallStackAllocationInfo AllocationInfo in FilteredCallstackList )
 				{
 					// Apply optional filter.
 					if( FStreamInfo.GlobalInstance.CallStackArray[ AllocationInfo.CallStackIndex ].RunFilters( FilterText, OwnerWindow.Options.ClassGroups, bFilterIn, OwnerWindow.SelectedMemoryPool ) )
 					{
-						TotalSize += AllocationInfo.Size;
-						TotalCount += AllocationInfo.Count;
+						TotalSize += AllocationInfo.TotalSize;
+						TotalCount += AllocationInfo.TotalCount;
 					}
 				}
 
 				// Clear out existing entries and add top 400.
 				ExclusiveListView.Items.Clear();
-				for( int CallStackIndex = 0; CallStackIndex < CallStackList.Count && ExclusiveListView.Items.Count <= MaximumEntries; CallStackIndex++ )
+				for( int CallStackIndex = 0; CallStackIndex < FilteredCallstackList.Count && ExclusiveListView.Items.Count <= MaximumEntries; CallStackIndex++ )
 				{
 					// Update progress bar.
 					if( CallStackCurrent >= NextProgressUpdate )
@@ -100,7 +110,7 @@ namespace MemoryProfiler2
 					}
 					CallStackCurrent++;
 
-					FCallStackAllocationInfo AllocationInfo = CallStackList[ CallStackIndex ];
+					FCallStackAllocationInfo AllocationInfo = FilteredCallstackList[ CallStackIndex ];
 
 					// Apply optional filter.
 					FCallStack CallStack = FStreamInfo.GlobalInstance.CallStackArray[ AllocationInfo.CallStackIndex ];
@@ -126,10 +136,13 @@ namespace MemoryProfiler2
 						}
 						while( UnhelpfulCallSites.Contains( FunctionName ) && FirstStackFrameIndex > 0 );
 
-						string SizeInKByte = String.Format( "{0:0}", ( float )AllocationInfo.Size / 1024 ).PadLeft( 10, ' ' );
-						string SizePercent = String.Format( "{0:0.00}", ( float )AllocationInfo.Size / TotalSize * 100 ).PadLeft( 10, ' ' );
-						string Count = String.Format( "{0:0}", AllocationInfo.Count ).PadLeft( 10, ' ' );
-						string CountPercent = String.Format( "{0:0.00}", ( float )AllocationInfo.Count / TotalCount * 100 ).PadLeft( 10, ' ' );
+						var AllocationSize = AllocationInfo.TotalSize;
+						var AllocationCount = AllocationInfo.TotalCount;
+
+						string SizeInKByte = String.Format( "{0:0}", ( float )AllocationSize / 1024 ).PadLeft( 10, ' ' );
+						string SizePercent = String.Format( "{0:0.00}", ( float )AllocationSize / TotalSize * 100 ).PadLeft( 10, ' ' );
+						string Count = String.Format( "{0:0}", AllocationCount ).PadLeft( 10, ' ' );
+						string CountPercent = String.Format( "{0:0.00}", ( float )AllocationCount / TotalCount * 100 ).PadLeft( 10, ' ' );
 						string GroupName = ( CallStack.Group != null ) ? CallStack.Group.Name : "Ungrouped";
 
 						string[] Row = new string[]
@@ -163,19 +176,19 @@ namespace MemoryProfiler2
 		/// <summary> Compare helper function, sorting FCallStackAllocation by size. </summary>
 		private static int CompareSize( FCallStackAllocationInfo A, FCallStackAllocationInfo B )
 		{
-			return Math.Sign( B.Size - A.Size );
+			return Math.Sign( B.TotalSize - A.TotalSize);
 		}
 
 		/// <summary> Compare helper function, sorting FCallStackAllocation by abs(size). </summary>
         private static int CompareAbsSize(FCallStackAllocationInfo A, FCallStackAllocationInfo B)
         {
-            return Math.Sign(Math.Abs(B.Size) - Math.Abs(A.Size));
+            return Math.Sign(Math.Abs(B.TotalSize) - Math.Abs(A.TotalSize));
         }
 
 		/// <summary> Compare helper function, sorting FCallStackAllocation by count. </summary>
 		private static int CompareCount( FCallStackAllocationInfo A, FCallStackAllocationInfo B )
 		{
-			return Math.Sign( B.Count - A.Count );
+			return Math.Sign( B.TotalCount - A.TotalCount);
 		}	
 	};
 }
\ No newline at end of file
diff --git a/Engine/Source/Programs/MemoryProfiler2/FilterTagsDialog.cs b/Engine/Source/Programs/MemoryProfiler2/FilterTagsDialog.cs
new file mode 100644
--- /dev/null
+++ b/Engine/Source/Programs/MemoryProfiler2/FilterTagsDialog.cs
@@ -0,0 +1,282 @@
+﻿using System;
+using System.Collections.Generic;
+using System.ComponentModel;
+using System.Data;
+using System.Drawing;
+using System.Linq;
+using System.Text;
+using System.Windows.Forms;
+
+namespace MemoryProfiler2
+{
+	public partial class FilterTagsDialog : Form
+	{
+		public FilterTagsDialog()
+		{
+			InitializeComponent();
+
+			FilterTagsTreeView.TreeViewNodeSorter = null; // Data is pre-sorted
+			FilterTagsTreeView.CheckBoxes = true;
+
+			FilterTagsSearch.TextChanged += FilterTagsSearch_TextChanged;
+			FilterTagsSearch.KeyDown += FilterTagsSearch_KeyDown;
+			FilterTagsSearch.LostFocus += FilterTagsSearch_LostFocus;
+
+			ResetFilterButton.Click += ResetFilterButton_Click;
+
+			ActiveTags = new HashSet<FAllocationTag>();
+
+			FilterTextTimer = new Timer();
+			FilterTextTimer.Interval = 1000;
+			FilterTextTimer.Tick += FilterTextTimer_Tick;
+		}
+
+		public HashSet<FAllocationTag> GetActiveTags()
+		{
+			return new HashSet<FAllocationTag>(ActiveTags);
+		}
+
+		public void PopulateTagFilterHierarchy()
+		{
+			FilterTagsSearch.Text = String.Empty;
+			LastFilterText = String.Empty;
+			PopulateTagFilterHierarchy(null);
+		}
+
+		private void PopulateTagFilterHierarchy(string[] InFilterTerms)
+		{
+			UseWaitCursor = true;
+
+			// Build the internal tree
+			{
+				InternalTreeRoot = new FilterTagsTreeNode(null);
+
+				var RootNode = FStreamInfo.GlobalInstance.TagHierarchy.GetRootNode();
+				BuildInternalTree(InFilterTerms, RootNode, InternalTreeRoot);
+			}
+
+			// Build the WinForm tree
+			if (FilterTagsTreeView.Handle != null) // This if statement exists to force the native tree view handle to be created
+			{
+				bool bIsFiltering = InFilterTerms != null && InFilterTerms.Length > 0;
+
+				FilterTagsTreeView.BeginUpdate();
+
+				FilterTagsTreeView.AfterCheck -= FilterTagsTreeView_AfterCheck;
+				FilterTagsTreeView.BeforeExpand -= FilterTagsTreeView_BeforeExpand;
+
+				FilterTagsTreeView.Nodes.Clear();
+				BuildFilterTagsTree(InternalTreeRoot, FilterTagsTreeView.Nodes, bIsFiltering);
+				UpdateCheckboxState(FilterTagsTreeView.Nodes);
+
+				if (bIsFiltering)
+				{
+					FilterTagsTreeView.ExpandAll();
+				}
+
+				FilterTagsTreeView.AfterCheck += FilterTagsTreeView_AfterCheck;
+				FilterTagsTreeView.BeforeExpand += FilterTagsTreeView_BeforeExpand;
+
+				FilterTagsTreeView.EndUpdate();
+			}
+
+			UseWaitCursor = false;
+		}
+
+		private bool BuildInternalTree(string[] InFilterTerms, FAllocationTagHierarchy.INode InParentAllocTagNode, FilterTagsTreeNode InParentInternalNode)
+		{
+			bool bIsFiltering = InFilterTerms != null && InFilterTerms.Length > 0;
+
+			foreach (var ChildAllocTagNode in InParentAllocTagNode.GetChildNodes())
+			{
+				var ChildInternalNode = new FilterTagsTreeNode(null, ChildAllocTagNode); // Don't set the parent yet as we might not add it to the tree
+				if (BuildInternalTree(InFilterTerms, ChildAllocTagNode, ChildInternalNode))
+				{
+					// Link the node now
+					ChildInternalNode.Parent = InParentInternalNode;
+					InParentInternalNode.Children.Add(ChildInternalNode);
+				}
+			}
+
+			foreach (var ChildAllocTag in InParentAllocTagNode.GetChildTags())
+			{
+				var AllocTag = ChildAllocTag.GetTag().Value;
+				bool bChildTagNamePassesFilter = true;
+
+				if (bIsFiltering)
+				{
+					var FullTagName = AllocTag.Name.ToLower();
+					foreach (var FilterTerm in InFilterTerms)
+					{
+						if (!FullTagName.Contains(FilterTerm))
+						{
+							bChildTagNamePassesFilter = false;
+							break;
+						}
+					}
+				}
+
+				if (bChildTagNamePassesFilter)
+				{
+					var ChildInternalNode = new FilterTagsTreeNode(InParentInternalNode, ChildAllocTag);
+				}
+			}
+
+			return InParentInternalNode.Children.Count > 0;
+		}
+
+		private void BuildFilterTagsTree(FilterTagsTreeNode InParentInternalNode, TreeNodeCollection InTreeNodesList, bool bRecursive)
+		{
+			foreach (var ChildInternalNode in InParentInternalNode.Children)
+			{
+				var ChildTreeNode = new TreeNode(ChildInternalNode.AllocTagNode.GetNodeName());
+				ChildTreeNode.Tag = ChildInternalNode;
+
+				if (bRecursive)
+				{
+					BuildFilterTagsTree(ChildInternalNode, ChildTreeNode.Nodes, bRecursive);
+				}
+				else
+				{
+					// Add a dummy child node if needed so that we still get an expansion arrow
+					if (ChildInternalNode.Children.Count > 0)
+					{
+						ChildTreeNode.Nodes.Add(new TreeNode());
+					}
+				}
+
+				InTreeNodesList.Add(ChildTreeNode);
+			}
+		}
+
+		private void UpdateCheckboxState(TreeNodeCollection InTreeNodesList)
+		{
+			foreach (TreeNode TreeNode in InTreeNodesList)
+			{
+				var InternalNode = TreeNode.Tag as FilterTagsTreeNode;
+				if (InternalNode != null)
+				{
+					if (InternalNode.AllocTagNode.GetTag().HasValue)
+					{
+						TreeNode.Checked = ActiveTags.Contains(InternalNode.AllocTagNode.GetTag().Value);
+					}
+					else
+					{
+						// Hide the checkbox on non-tag nodes
+						FilterTagsTreeView.SetCheckBoxImageState(TreeNode, 0);
+					}
+				}
+
+				UpdateCheckboxState(TreeNode.Nodes);
+			}
+		}
+
+		private void ApplyTextFilter()
+		{
+			FilterTextTimer.Stop();
+
+			if (LastFilterText == null || LastFilterText != FilterTagsSearch.Text)
+			{
+				LastFilterText = FilterTagsSearch.Text;
+
+				var FilterTerms = FilterTagsSearch.Text.ToLower().Split(FilterTermSeparatorCharacters, StringSplitOptions.RemoveEmptyEntries);
+				PopulateTagFilterHierarchy(FilterTerms);
+			}
+		}
+
+		private void FilterTagsSearch_TextChanged(object sender, System.EventArgs e)
+		{
+			FilterTextTimer.Stop();
+			FilterTextTimer.Start();
+		}
+
+		private void FilterTagsSearch_KeyDown(object sender, KeyEventArgs e)
+		{
+			if (e.KeyCode == Keys.Enter)
+			{
+				ApplyTextFilter();
+			}
+		}
+
+		private void FilterTagsSearch_LostFocus(object sender, EventArgs e)
+		{
+			ApplyTextFilter();
+		}
+
+		private void FilterTagsTreeView_AfterCheck(object sender, TreeViewEventArgs e)
+		{
+			var InternalNode = e.Node.Tag as FilterTagsTreeNode;
+			if (InternalNode != null)
+			{
+				if (e.Node.Checked)
+				{
+					ActiveTags.Add(InternalNode.AllocTagNode.GetTag().Value);
+				}
+				else
+				{
+					ActiveTags.Remove(InternalNode.AllocTagNode.GetTag().Value);
+				}
+			}
+		}
+
+		private void FilterTagsTreeView_BeforeExpand(object sender, TreeViewCancelEventArgs e)
+		{
+			if (e.Node.Nodes.Count == 1 && e.Node.Nodes[0].Tag == null)
+			{
+				// Replace the dummy child with real children now that we're being expanded
+				e.Node.TreeView.BeginUpdate();
+				e.Node.Nodes.Clear();
+				BuildFilterTagsTree((FilterTagsTreeNode)e.Node.Tag, e.Node.Nodes, false);
+				UpdateCheckboxState(e.Node.Nodes);
+				e.Node.TreeView.EndUpdate();
+			}
+		}
+
+		private void ResetFilterButton_Click(object sender, EventArgs e)
+		{
+			ActiveTags.Clear();
+			FilterTagsSearch.Text = String.Empty;
+			PopulateTagFilterHierarchy(null);
+		}
+
+		private void FilterTextTimer_Tick(object sender, EventArgs e)
+		{
+			ApplyTextFilter();
+		}
+
+		private FilterTagsTreeNode InternalTreeRoot = null;
+
+		private HashSet<FAllocationTag> ActiveTags;
+
+		private Timer FilterTextTimer = null;
+
+		private string LastFilterText = null;
+
+		private static char[] FilterTermSeparatorCharacters = { ' ' };
+	}
+
+	class FilterTagsTreeNode
+	{
+		public FilterTagsTreeNode Parent = null;
+
+		public List<FilterTagsTreeNode> Children = new List<FilterTagsTreeNode>();
+
+		public FAllocationTagHierarchy.INode AllocTagNode;
+
+		public FilterTagsTreeNode(FAllocationTagHierarchy.INode InAllocTagNode)
+		{
+			AllocTagNode = InAllocTagNode;
+		}
+
+		public FilterTagsTreeNode(FilterTagsTreeNode InParent, FAllocationTagHierarchy.INode InAllocTagNode)
+		{
+			Parent = InParent;
+			if (Parent != null)
+			{
+				Parent.Children.Add(this);
+			}
+
+			AllocTagNode = InAllocTagNode;
+		}
+	}
+}
diff --git a/Engine/Source/Programs/MemoryProfiler2/HistogramParser.cs b/Engine/Source/Programs/MemoryProfiler2/HistogramParser.cs
--- a/Engine/Source/Programs/MemoryProfiler2/HistogramParser.cs
+++ b/Engine/Source/Programs/MemoryProfiler2/HistogramParser.cs
@@ -128,10 +128,22 @@ namespace MemoryProfiler2
 				long Size = 0;
 				int Count = 0;
 
-				foreach( FCallStackAllocationInfo AllocationInfo in CallStackList )
+				bool bFilterIn = OwnerWindow.IsFilteringIn();
+
+				var FilteredCallstackList = new List<FCallStackAllocationInfo>(CallStackList.Count);
+				foreach (var AllocationInfo in CallStackList)
+				{
+					var FilteredAllocationInfo = AllocationInfo.GetAllocationInfoForTags(OwnerWindow.GetTagsFilter(), bFilterIn);
+					if (FilteredAllocationInfo.TotalCount != 0)
+					{
+						FilteredCallstackList.Add(FilteredAllocationInfo);
+					}
+				}
+
+				foreach ( FCallStackAllocationInfo AllocationInfo in FilteredCallstackList )
 				{
 					// Update progress bar.
-					if( CallStackCurrent >= NextProgressUpdate )
+					if ( CallStackCurrent >= NextProgressUpdate )
 					{
 						OwnerWindow.ToolStripProgressBar.PerformStep();
 						NextProgressUpdate += ProgressInterval;
@@ -139,8 +151,8 @@ namespace MemoryProfiler2
 					}
 					CallStackCurrent++;
 
-					FCallStack OriginalCallStack = FStreamInfo.GlobalInstance.CallStackArray[ AllocationInfo.CallStackIndex ];
-					if( OriginalCallStack.RunFilters( FilterText, CallStackGroups, OwnerWindow.IsFilteringIn(), OwnerWindow.SelectedMemoryPool ) )
+					FCallStack OriginalCallStack = FStreamInfo.GlobalInstance.CallStackArray[AllocationInfo.CallStackIndex ];
+					if( OriginalCallStack.RunFilters( FilterText, CallStackGroups, bFilterIn, OwnerWindow.SelectedMemoryPool ) )
 					{
 						bool bFound = false;
 						int Column = FMemoryPoolInfo.GetMemoryPoolHistogramColumn( OriginalCallStack.MemoryPool );
@@ -172,8 +184,8 @@ HackyBreakAll:
 						}
 					}
 
-					Size += AllocationInfo.Size;
-					Count += AllocationInfo.Count;
+					Size += AllocationInfo.TotalSize;
+					Count += AllocationInfo.TotalCount;
 				}
 			}
 
@@ -716,7 +728,7 @@ HackyBreakAll:
 			if (BatchAddingCount == 0)
 			{
 				// Sorting largest -> smallest
-				CallStackList.Sort((First, Second) => Math.Sign(First.Size - Second.Size));
+				CallStackList.Sort((First, Second) => Math.Sign(First.TotalSize - Second.TotalSize));
 			}
 		}
 
@@ -729,7 +741,7 @@ HackyBreakAll:
 			{
 				for (int Index = 0; Index < CallStackList.Count; Index++)
 				{
-					if (CallStackList[Index].Size > AllocationInfo.Size)
+					if (CallStackList[Index].TotalSize > AllocationInfo.TotalSize)
 					{
 						CallStackList.Insert(Index, AllocationInfo);
 						bInserted = true;
@@ -743,8 +755,8 @@ HackyBreakAll:
 				CallStackList.Add(AllocationInfo);
 			}
 
-			MemorySize += AllocationInfo.Size;
-			AllocationCount += AllocationInfo.Count;
+			MemorySize += AllocationInfo.TotalSize;
+			AllocationCount += AllocationInfo.TotalCount;
 		}
 	}
 }
diff --git a/Engine/Source/Programs/MemoryProfiler2/MainWindow.cs b/Engine/Source/Programs/MemoryProfiler2/MainWindow.cs
--- a/Engine/Source/Programs/MemoryProfiler2/MainWindow.cs
+++ b/Engine/Source/Programs/MemoryProfiler2/MainWindow.cs
@@ -45,6 +45,9 @@ namespace MemoryProfiler2
 		/// <summary> Progress indicator window. </summary>
 		public SlowProgressDialog ProgressDialog = new SlowProgressDialog();
 
+		/// <summary> Tags filter dialog. </summary>
+		public FilterTagsDialog FilterTagsDialog = new FilterTagsDialog();
+
 		/// <summary> Selected memory pool filter. </summary>
 		public EMemoryPool SelectedMemoryPool = EMemoryPool.MEMPOOL_None;
 
@@ -181,18 +184,25 @@ namespace MemoryProfiler2
 			FilterClassesDropDownButton.DropDownItems.Add( Separator );
 
 			ToolStripMenuItem AllMenuItem = new ToolStripMenuItem( "All" );
-			AllMenuItem.Click += new System.EventHandler( FilterAllClick );
+			AllMenuItem.Click += new System.EventHandler( FilterAllClassGroupsClick );
 			AllMenuItem.MouseUp += ClassesToFilterMouseUp;
 			AllMenuItem.Tag = AllHistoryTag;
 
 			FilterClassesDropDownButton.DropDownItems.Add( AllMenuItem );
 
 			ToolStripMenuItem NoneMenuItem = new ToolStripMenuItem( "None" );
-			NoneMenuItem.Click += new System.EventHandler( FilterNoneClick );
+			NoneMenuItem.Click += new System.EventHandler( FilterNoneClassGroupsClick );
 			FilterClassesDropDownButton.DropDownItems.Add( NoneMenuItem );
 		}
 		#endregion
 
+		private void FilterTagsButton_Click(object sender, EventArgs e)
+		{
+			FilterTagsDialog.PopulateTagFilterHierarchy();
+			FilterTagsDialog.ShowDialog();
+			UpdateFilteringState();
+		}
+
 		/// <summary> Assembly resolve method to pick correct StandaloneSymbolParser DLL. </summary>
 		private Assembly CurrentDomain_AssemblyResolve( Object sender, ResolveEventArgs args )
 		{
@@ -415,6 +425,11 @@ namespace MemoryProfiler2
 			UpdateStatus( "Displaying " + CurrentFilename );
 		}
 
+		public ISet<FAllocationTag> GetTagsFilter()
+		{
+			return CurrentState.AllocTags;
+		}
+
 		private void GoButton_Click( object sender, EventArgs e )
 		{
 			SetWaitCursor( true );
@@ -662,6 +677,7 @@ namespace MemoryProfiler2
 
 			FilterTypeSplitButton.Enabled = e.TabPage == TimeLineTabPage ? false : true;
 			FilterClassesDropDownButton.Enabled = e.TabPage == TimeLineTabPage ? false : true;
+			FilterTagsButton.Enabled = e.TabPage == CallGraphTabPage || e.TabPage == ExclusiveTabPage || e.TabPage == HistogramTabPage;
 			FilterTypeSplitButton.Enabled = e.TabPage == TimeLineTabPage ? false : true;
 
 			FilterLabel.Enabled = e.TabPage == TimeLineTabPage ? false : true;
@@ -855,7 +871,7 @@ namespace MemoryProfiler2
 		}
 
 		/// <summary> Checks all class group in the class filter menu. </summary>
-		private void FilterAllClick( object sender, EventArgs e )
+		private void FilterAllClassGroupsClick( object sender, EventArgs e )
 		{
 			SetAllClasses( true );
 			FilterClassesDropDownButton.ShowDropDown();
@@ -864,7 +880,7 @@ namespace MemoryProfiler2
 		}
 
 		/// <summary> Unchecks all class group in the class filter menu. </summary>
-		private void FilterNoneClick( object sender, EventArgs e )
+		private void FilterNoneClassGroupsClick( object sender, EventArgs e )
 		{
 			SetAllClasses( false );
 			FilterClassesDropDownButton.ShowDropDown();
@@ -2094,6 +2110,7 @@ namespace MemoryProfiler2
 			FC_GroupArray,
 			FC_TextFilter,
 			FC_MemoryPool,
+			FC_AllocTags,
 			Count,
 		}
 
@@ -2111,6 +2128,7 @@ namespace MemoryProfiler2
 			public List<bool> GroupArray;
 			public string TextFilter;
 			public EMemoryPool MemoryPool;
+			public HashSet<FAllocationTag> AllocTags;
 
 			/// <summary> Default constructor. </summary>
 			public FFilteringState( bool bFake )
@@ -2124,6 +2142,7 @@ namespace MemoryProfiler2
 				GroupArray = new List<bool>( 40 );
 				TextFilter = "";
 				MemoryPool = EMemoryPool.MEMPOOL_None;
+				AllocTags = new HashSet<FAllocationTag>();
 			}
 
 			/// <summary> Creates a new copy of this class. </summary>
@@ -2131,6 +2150,7 @@ namespace MemoryProfiler2
 			{
 				FFilteringState Copy = ( FFilteringState )MemberwiseClone();
 				Copy.GroupArray = new List<bool>( GroupArray );
+				Copy.AllocTags = new HashSet<FAllocationTag>( AllocTags );
 				return Copy;
 			}
 		};
@@ -2185,6 +2205,7 @@ namespace MemoryProfiler2
 
 				CurrentState.TextFilter = FilterTextBox.Text;
 				CurrentState.MemoryPool = SelectedMemoryPool;
+				CurrentState.AllocTags = FilterTagsDialog.GetActiveTags();
 
 				//bool bDiffFound = !CurrentState.CompareWith( StartingState );
 
@@ -2223,6 +2244,23 @@ namespace MemoryProfiler2
 					bEquals[ ( int )EFilteringControl.FC_GroupArray ] = false;
 				}
 
+				bEquals[ ( int )EFilteringControl.FC_AllocTags ] = true;
+				if( CurrentState.AllocTags.Count == StartingState.AllocTags.Count )
+				{
+					foreach (var AllocTag in CurrentState.AllocTags)
+					{
+						if (!StartingState.AllocTags.Contains(AllocTag))
+						{
+							bEquals[ ( int )EFilteringControl.FC_AllocTags ] = false;
+							break;
+						}
+					}
+				}
+				else
+				{
+					bEquals[ ( int )EFilteringControl.FC_AllocTags ] = false;
+				}
+
 				bool bEqualProperties = true;
 				for( int Eq = 0; Eq < ( int )EFilteringControl.Count; Eq++ )
 				{
@@ -2259,6 +2297,7 @@ namespace MemoryProfiler2
 
 			StartingState.TextFilter = FilterTextBox.Text;
 			StartingState.MemoryPool = SelectedMemoryPool;
+			StartingState.AllocTags = FilterTagsDialog.GetActiveTags();
 
 			CurrentState = StartingState.DeepCopy();
 			UpdateFilteringState();
@@ -2276,6 +2315,7 @@ namespace MemoryProfiler2
 			FilteringControl[ ( int )EFilteringControl.FC_GroupArray ] = FilterClassesDropDownButton;
 			FilteringControl[ ( int )EFilteringControl.FC_TextFilter ] = FilterLabel;
 			FilteringControl[ ( int )EFilteringControl.FC_MemoryPool ] = MemoryPoolFilterButton;
+			FilteringControl[ ( int )EFilteringControl.FC_AllocTags ] = FilterTagsButton;
 
 			FilterClassesDropDownButton.DropDown.Renderer.RenderItemText += new ToolStripItemTextRenderEventHandler( Renderer_RenderItemText );
 
diff --git a/Engine/Source/Programs/MemoryProfiler2/MemoryProfiler2.csproj b/Engine/Source/Programs/MemoryProfiler2/MemoryProfiler2.csproj
--- a/Engine/Source/Programs/MemoryProfiler2/MemoryProfiler2.csproj
+++ b/Engine/Source/Programs/MemoryProfiler2/MemoryProfiler2.csproj
@@ -84,6 +84,7 @@
     </Reference>
   </ItemGroup>
   <ItemGroup>
+    <Compile Include="AllocationTags.cs" />
     <Compile Include="BinaryReaderBigEndian.cs" />
     <Compile Include="CallGraphTreeViewParser.cs" />
     <Compile Include="CallStack.cs" />
@@ -95,6 +96,12 @@
       <SubType>Component</SubType>
     </Compile>
     <Compile Include="ExclusiveListViewParser.cs" />
+    <Compile Include="FilterTagsDialog.cs">
+      <SubType>Form</SubType>
+    </Compile>
+    <Compile Include="FilterTagsDialog.Designer.cs">
+      <DependentUpon>FilterTagsDialog.cs</DependentUpon>
+    </Compile>
     <Compile Include="FindDialog.cs">
       <SubType>Form</SubType>
     </Compile>
@@ -126,6 +133,9 @@
     </Compile>
     <Compile Include="Program.cs" />
     <Compile Include="Properties\AssemblyInfo.cs" />
+    <EmbeddedResource Include="FilterTagsDialog.resx">
+      <DependentUpon>FilterTagsDialog.cs</DependentUpon>
+    </EmbeddedResource>
     <EmbeddedResource Include="FindDialog.resx">
       <DependentUpon>FindDialog.cs</DependentUpon>
     </EmbeddedResource>
diff --git a/Engine/Source/Programs/MemoryProfiler2/ProfileDataHeader.cs b/Engine/Source/Programs/MemoryProfiler2/ProfileDataHeader.cs
--- a/Engine/Source/Programs/MemoryProfiler2/ProfileDataHeader.cs
+++ b/Engine/Source/Programs/MemoryProfiler2/ProfileDataHeader.cs
@@ -44,10 +44,17 @@ namespace MemoryProfiler2
 		public UInt64 CallStackTableOffset;
 		/** Number of callstack entries.							*/
 		public UInt64 CallStackTableEntries;
+
+		/** Offset in file for tags table.							*/
+		public UInt64 TagsTableOffset;
+		/** Number of tags entries.									*/
+		public UInt64 TagsTableEntries;
+
 		/** The file offset for module information.					*/
 		public UInt64 ModulesOffset;
 		/** The number of module entries.							*/
 		public UInt64 ModuleEntries;
+
 		/** Number of data files the stream spans.					*/
 		public UInt64 NumDataFiles;
 
@@ -117,6 +124,18 @@ namespace MemoryProfiler2
 					CallStackTableOffset = BinaryStream.ReadUInt64();
 					CallStackTableEntries = BinaryStream.ReadUInt64();
 
+					if (Version >= 7)
+					{
+						// Tags table offset and number of entries.
+						TagsTableOffset = BinaryStream.ReadUInt64();
+						TagsTableEntries = BinaryStream.ReadUInt64();
+					}
+					else
+					{
+						TagsTableOffset = 0;
+						TagsTableEntries = 0;
+					}
+
 					ModulesOffset = BinaryStream.ReadUInt64();
 					ModuleEntries = BinaryStream.ReadUInt64();
 
diff --git a/Engine/Source/Programs/MemoryProfiler2/Properties/AssemblyInfo.cs b/Engine/Source/Programs/MemoryProfiler2/Properties/AssemblyInfo.cs
--- a/Engine/Source/Programs/MemoryProfiler2/Properties/AssemblyInfo.cs
+++ b/Engine/Source/Programs/MemoryProfiler2/Properties/AssemblyInfo.cs
@@ -3,6 +3,6 @@
  */
 using System.Reflection;
 
-[assembly: AssemblyTitle( "Memory Profiler v2.0.2" )]
+[assembly: AssemblyTitle( "Memory Profiler v2.1" )]
 [assembly: AssemblyDescription( "A tool to process the output from the in engine memory profiling tools." )]
 [assembly: AssemblyConfiguration("")]
diff --git a/Engine/Source/Programs/MemoryProfiler2/StreamInfo.cs b/Engine/Source/Programs/MemoryProfiler2/StreamInfo.cs
--- a/Engine/Source/Programs/MemoryProfiler2/StreamInfo.cs
+++ b/Engine/Source/Programs/MemoryProfiler2/StreamInfo.cs
@@ -21,6 +21,12 @@ namespace MemoryProfiler2
 		/// <summary> Meta-data associated with this stream. </summary>
 		public Dictionary<string, string> MetaData;
 
+		/// <summary> Array of unique tags. Code has fixed indexes into it. </summary>
+		public List<FAllocationTags> TagsArray;
+
+		/// <summary> Hierarchy of tags built as they are parsed. </summary>
+		public FAllocationTagHierarchy TagHierarchy = new FAllocationTagHierarchy();
+
 		/// <summary> Array of unique names. Code has fixed indexes into it. </summary>
 		private List<string> InternalNameArray;
 		private Dictionary<string, int> InternalNameIndexLookupMap;
@@ -103,6 +109,7 @@ namespace MemoryProfiler2
 		{
 			PlatformName = Header.PlatformName;
 			MetaData = new Dictionary<string, string>( (int)Header.MetaDataTableEntries );
+			TagsArray = new List<FAllocationTags>( (int)Header.TagsTableEntries );
 			InternalNameArray = new List<string>( (int)Header.NameTableEntries );
 			InternalNameIndexLookupMap = new Dictionary<string, int>( (int)Header.NameTableEntries );
 			CallStackArray = new List<FCallStack>( (int)Header.CallStackTableEntries );
diff --git a/Engine/Source/Programs/MemoryProfiler2/StreamParser.cs b/Engine/Source/Programs/MemoryProfiler2/StreamParser.cs
--- a/Engine/Source/Programs/MemoryProfiler2/StreamParser.cs
+++ b/Engine/Source/Programs/MemoryProfiler2/StreamParser.cs
@@ -113,6 +113,14 @@ namespace MemoryProfiler2
 			// Keep track of current position as it's where the token stream starts.
 			long TokenStreamOffset = ParserFileStream.Position;
 
+			// Seek to name table and serialize it.
+			ParserFileStream.Seek((Int64)Header.NameTableOffset, SeekOrigin.Begin);
+			for (UInt64 NameIndex = 0; NameIndex < Header.NameTableEntries; NameIndex++)
+			{
+				int InsertedNameIndex = FStreamInfo.GlobalInstance.GetNameIndex(ReadString(BinaryStream), true);
+				Debug.Assert((int)NameIndex == InsertedNameIndex);
+			}
+
 			if (Header.Version >= 6)
 			{
 				// Seek to meta-data table and serialize it.
@@ -125,13 +133,18 @@ namespace MemoryProfiler2
 				}
 			}
 
-			// Seek to name table and serialize it.
-			ParserFileStream.Seek((Int64)Header.NameTableOffset,SeekOrigin.Begin);
-			for(UInt64 NameIndex = 0;NameIndex < Header.NameTableEntries;NameIndex++)
+			FStreamInfo.GlobalInstance.TagHierarchy.InitializeHierarchy();
+			if (Header.Version >= 7)
 			{
-				int InsertedNameIndex = FStreamInfo.GlobalInstance.GetNameIndex( ReadString( BinaryStream ), true );
-				Debug.Assert((int)NameIndex == InsertedNameIndex);
+				// Seek to tags table and serialize it.
+				ParserFileStream.Seek((Int64)Header.TagsTableOffset, SeekOrigin.Begin);
+				for (UInt64 TagsIndex = 0; TagsIndex < Header.TagsTableEntries; TagsIndex++)
+				{
+					string TagsString = ReadString(BinaryStream);
+					FStreamInfo.GlobalInstance.TagsArray.Add(new FAllocationTags(TagsString));
+				}
 			}
+			FStreamInfo.GlobalInstance.TagHierarchy.FinalizeHierarchy();
 
 			// Seek to callstack address array and serialize it.                
 			ParserFileStream.Seek( (Int64)Header.CallStackAddressTableOffset, SeekOrigin.Begin );
@@ -276,9 +289,9 @@ namespace MemoryProfiler2
 
 			// Snapshot used for parsing. A copy will be made if a special token is encountered. Otherwise it
 			// will be returned as the only snaphot at the end.
-			FStreamSnapshot Snapshot = new FStreamSnapshot("End");
-			List<FStreamSnapshot> SnapshotList = new List<FStreamSnapshot>();
-            Dictionary<ulong, FCallStackAllocationInfo> PointerToPointerInfoMap = new Dictionary<ulong, FCallStackAllocationInfo>();
+			var Snapshot = new FStreamSnapshot("End");
+			var SnapshotList = new List<FStreamSnapshot>();
+            var PointerToPointerInfoMap = new Dictionary<ulong, FLiveAllocationInfo>();
 
 			// Seek to beginning of token stream.
 			ParserFileStream.Seek(TokenStreamOffset, SeekOrigin.Begin);
@@ -423,7 +436,7 @@ namespace MemoryProfiler2
 								// Either USE_GLOBAL_REALLOC_ZERO_PTR is not being used, or we're not
 								// trying to free the ReallocZeroPtr.
 
-								FCallStackAllocationInfo FreedAllocInfo;
+								FLiveAllocationInfo FreedAllocInfo;
 								if( HandleFree( Token, Snapshot, PointerToPointerInfoMap, out FreedAllocInfo ) )
 								{
 									FCallStack PreviousCallStack = FStreamInfo.GlobalInstance.CallStackArray[ FreedAllocInfo.CallStackIndex ];
@@ -451,7 +464,7 @@ namespace MemoryProfiler2
 
 							if( Token.OldPointer != 0 )
 							{
-								FCallStackAllocationInfo FreedAllocInfo;
+								FLiveAllocationInfo FreedAllocInfo;
 								if( HandleFree( Token, Snapshot, PointerToPointerInfoMap, out FreedAllocInfo ) )
 								{
 									PreviousCallstack = FStreamInfo.GlobalInstance.CallStackArray[ FreedAllocInfo.CallStackIndex ];
@@ -734,7 +747,7 @@ namespace MemoryProfiler2
 			{
 				while( PreviousSnapshot.LifetimeCallStackList.Count < Snapshot.LifetimeCallStackList.Count )
 				{
-					PreviousSnapshot.LifetimeCallStackList.Add( new FCallStackAllocationInfo( 0, PreviousSnapshot.LifetimeCallStackList.Count, 0 ) );
+					PreviousSnapshot.LifetimeCallStackList.Add( new FCallStackAllocationInfo( 0, PreviousSnapshot.LifetimeCallStackList.Count, 0, -1 ) );
 				}
 			}
 
@@ -838,10 +851,10 @@ namespace MemoryProfiler2
 		}
 
 		/// <summary> Updates internal state with allocation. </summary>
-        private static void HandleMalloc(FStreamToken StreamToken, FStreamSnapshot Snapshot, Dictionary<ulong, FCallStackAllocationInfo> PointerToPointerInfoMap)
+        private static void HandleMalloc(FStreamToken StreamToken, FStreamSnapshot Snapshot, Dictionary<ulong, FLiveAllocationInfo> PointerToPointerInfoMap)
 		{
 			// Keep track of size associated with pointer and also current callstack.
-            FCallStackAllocationInfo PointerInfo = new FCallStackAllocationInfo(StreamToken.Size, StreamToken.CallStackIndex, 1);
+			var PointerInfo = new FLiveAllocationInfo(StreamToken.Size, StreamToken.CallStackIndex, StreamToken.TagsIndex);
 
             if (PointerToPointerInfoMap.ContainsKey(StreamToken.Pointer))
             {
@@ -860,11 +873,10 @@ namespace MemoryProfiler2
             // Add size to lifetime churn tracking.
             while (StreamToken.CallStackIndex >= Snapshot.LifetimeCallStackList.Count)
 			{
-                Snapshot.LifetimeCallStackList.Add(new FCallStackAllocationInfo(0, Snapshot.LifetimeCallStackList.Count, 0));
+                Snapshot.LifetimeCallStackList.Add(new FCallStackAllocationInfo(0, Snapshot.LifetimeCallStackList.Count, 0, -1));
 			}
 
-            //@todo: Sadly, we have to do all this ugly shuffling because of the way lists of structs work in C#
-			Snapshot.LifetimeCallStackList[ StreamToken.CallStackIndex ] = Snapshot.LifetimeCallStackList[ StreamToken.CallStackIndex ].Add( StreamToken.Size, 1 );
+			Snapshot.LifetimeCallStackList[StreamToken.CallStackIndex].Add(StreamToken.Size, 1, StreamToken.TagsIndex);
 
 			if( Snapshot.AllocationSize > Snapshot.AllocationMaxSize )
 			{
@@ -883,7 +895,7 @@ namespace MemoryProfiler2
 		}
 
 		/// <summary> Updates internal state with free. </summary>
-        private static bool HandleFree(FStreamToken StreamToken, FStreamSnapshot Snapshot, Dictionary<ulong, FCallStackAllocationInfo> PointerToPointerInfoMap, out FCallStackAllocationInfo FreedAllocInfo)
+        private static bool HandleFree(FStreamToken StreamToken, FStreamSnapshot Snapshot, Dictionary<ulong, FLiveAllocationInfo> PointerToPointerInfoMap, out FLiveAllocationInfo FreedAllocInfo)
 		{
             if (!PointerToPointerInfoMap.TryGetValue(StreamToken.Pointer, out FreedAllocInfo))
             {
diff --git a/Engine/Source/Programs/MemoryProfiler2/StreamSnapshot.cs b/Engine/Source/Programs/MemoryProfiler2/StreamSnapshot.cs
--- a/Engine/Source/Programs/MemoryProfiler2/StreamSnapshot.cs
+++ b/Engine/Source/Programs/MemoryProfiler2/StreamSnapshot.cs
@@ -11,49 +11,296 @@ using System.Linq;
 
 namespace MemoryProfiler2
 {
-	/// <summary> Helper struct encapsulating information for a callstack and associated allocation size. Shared with stream parser. </summary>
-	public struct FCallStackAllocationInfo
+	/// <summary> Holds the size and callstack associated with a live allocation. </summary>
+	public struct FLiveAllocationInfo
 	{
 		/// <summary> Size of allocation. </summary>
-		public long Size;
+		public readonly long Size;
 
 		/// <summary> Index of callstack that performed allocation. </summary>
-		public int CallStackIndex;
+		public readonly int CallStackIndex;
+
+		/// <summary> Tags index associated with this allocation. </summary>
+		public readonly int TagsIndex;
+
+		/// <summary> Constructor, initializing all member variables to passed in values. </summary>
+		public FLiveAllocationInfo(long InSize, int InCallStackIndex, int InTagsIndex)
+		{
+			Size = InSize;
+			CallStackIndex = InCallStackIndex;
+			TagsIndex = InTagsIndex;
+		}
+	}
+
+	/// <summary> Holds the allocation size and count associated with a given group of tags. </summary>
+	public struct FCallStackTagsAllocationInfo
+	{
+		/// <summary> Size of allocation. </summary>
+		public long Size;
 
 		/// <summary> Number of allocations. </summary>
 		public int Count;
 
 		/// <summary> Constructor, initializing all member variables to passed in values. </summary>
-		public FCallStackAllocationInfo( long InSize, int InCallStackIndex, int InCount )
+		public FCallStackTagsAllocationInfo(long InSize, int InCount)
 		{
 			Size = InSize;
-			CallStackIndex = InCallStackIndex;
 			Count = InCount;
 		}
 
+		public static FCallStackTagsAllocationInfo operator+(FCallStackTagsAllocationInfo LHS, FCallStackTagsAllocationInfo RHS)
+		{
+			return new FCallStackTagsAllocationInfo(LHS.Size + RHS.Size, LHS.Count + RHS.Count);
+		}
+
+		public static FCallStackTagsAllocationInfo operator-(FCallStackTagsAllocationInfo LHS, FCallStackTagsAllocationInfo RHS)
+		{
+			return new FCallStackTagsAllocationInfo(LHS.Size - RHS.Size, LHS.Count - RHS.Count);
+		}
+	}
+
+	/// <summary> Helper class encapsulating information for a callstack and associated allocation size. Shared with stream parser. </summary>
+	public class FCallStackAllocationInfo
+	{
+		/// <summary> Index of callstack that performed allocation. </summary>
+		public readonly int CallStackIndex;
+
+		/// <summary> Allocation info not associated with any tags. </summary>
+		FCallStackTagsAllocationInfo UntaggedAllocationInfo;
+
+		/// <summary> Mapping between a tags index and the associated allocation info. </summary>
+		Dictionary<int, FCallStackTagsAllocationInfo> TaggedAllocationInfo;
+
+		/// <summary> Total size of allocation. </summary>
+		public long TotalSize
+		{
+			get
+			{
+				long TotalSize = UntaggedAllocationInfo.Size;
+				if (TaggedAllocationInfo != null)
+				{
+					foreach (var TaggedAllocationInfoPair in TaggedAllocationInfo)
+					{
+						TotalSize += TaggedAllocationInfoPair.Value.Size;
+					}
+				}
+				return TotalSize;
+			}
+		}
+
+		/// <summary> Total number of allocations. </summary>
+		public int TotalCount
+		{
+			get
+			{
+				int TotalCount = UntaggedAllocationInfo.Count;
+				if (TaggedAllocationInfo != null)
+				{
+					foreach (var TaggedAllocationInfoPair in TaggedAllocationInfo)
+					{
+						TotalCount += TaggedAllocationInfoPair.Value.Count;
+					}
+				}
+				return TotalCount;
+			}
+		}
+
+		/// <summary> Constructor, initializing all member variables to passed in values. </summary>
+		public FCallStackAllocationInfo( long InSize, int InCallStackIndex, int InCount, int InTagsIndex )
+		{
+			CallStackIndex = InCallStackIndex;
+			UntaggedAllocationInfo = new FCallStackTagsAllocationInfo(0, 0);
+			TaggedAllocationInfo = null;
+
+			if (InTagsIndex == -1)
+			{
+				UntaggedAllocationInfo += new FCallStackTagsAllocationInfo(InSize, InCount);
+			}
+			else
+			{
+				EnsureTaggedAllocationsAvailable();
+				TaggedAllocationInfo.Add(InTagsIndex, new FCallStackTagsAllocationInfo(InSize, InCount));
+			}
+		}
+
+		/// <summary> Private constructor used by DeepCopy. </summary>
+		private FCallStackAllocationInfo( int InCallStackIndex, FCallStackTagsAllocationInfo InUntaggedAllocationInfo, Dictionary<int, FCallStackTagsAllocationInfo> InTaggedAllocationInfo )
+		{
+			CallStackIndex = InCallStackIndex;
+			UntaggedAllocationInfo = InUntaggedAllocationInfo;
+			TaggedAllocationInfo = InTaggedAllocationInfo;
+		}
+
+		/// <summary> Performs a deep copy of the relevant data structures. </summary>
+		public FCallStackAllocationInfo DeepCopy()
+		{
+			return new FCallStackAllocationInfo( CallStackIndex, UntaggedAllocationInfo, TaggedAllocationInfo != null ? new Dictionary<int, FCallStackTagsAllocationInfo>(TaggedAllocationInfo) : null );
+		}
+
+		/// <summary> Call prior to adding data to TaggedAllocationInfo. </summary>
+		private void EnsureTaggedAllocationsAvailable()
+		{
+			if (TaggedAllocationInfo == null)
+			{
+				TaggedAllocationInfo = new Dictionary<int, FCallStackTagsAllocationInfo>(16);
+			}
+		}
+
 		/// <summary>
 		/// Adds the passed in size and count to this callstack info.
 		/// </summary>
 		/// <param name="SizeToAdd"> Size to add </param>
 		/// <param name="CountToAdd"> Count to add </param>
-		public FCallStackAllocationInfo Add( long SizeToAdd, int CountToAdd )
+		/// <param name="InTagsIndex"> Index of the tags associated with the allocation </param>
+		public void Add( long SizeToAdd, int CountToAdd, int InTagsIndex )
         {
-			return new FCallStackAllocationInfo( Size + SizeToAdd, CallStackIndex, Count + CountToAdd );
-        }
+			if (InTagsIndex == -1)
+			{
+				UntaggedAllocationInfo += new FCallStackTagsAllocationInfo(SizeToAdd, CountToAdd);
+			}
+			else
+			{
+				EnsureTaggedAllocationsAvailable();
+
+				FCallStackTagsAllocationInfo CurrentTagsAllocationInfo;
+				if (!TaggedAllocationInfo.TryGetValue(InTagsIndex, out CurrentTagsAllocationInfo))
+				{
+					CurrentTagsAllocationInfo = new FCallStackTagsAllocationInfo(0, 0);
+				}
+				TaggedAllocationInfo[InTagsIndex] = CurrentTagsAllocationInfo + new FCallStackTagsAllocationInfo(SizeToAdd, CountToAdd);
+			}
+		}
+
+		/// <summary>
+		/// Adds the passed callstack info to this callstack info.
+		/// </summary>
+		/// <param name="InOther"> Callstack info to add </param>
+		public void Add(FCallStackAllocationInfo InOther)
+		{
+			if (CallStackIndex != InOther.CallStackIndex)
+			{
+				throw new InvalidDataException();
+			}
+
+			UntaggedAllocationInfo += InOther.UntaggedAllocationInfo;
+
+			if (InOther.TaggedAllocationInfo != null)
+			{
+				EnsureTaggedAllocationsAvailable();
+
+				foreach (var TaggedAllocationInfoPair in InOther.TaggedAllocationInfo)
+				{
+					FCallStackTagsAllocationInfo CurrentTagsAllocationInfo;
+					if (!TaggedAllocationInfo.TryGetValue(TaggedAllocationInfoPair.Key, out CurrentTagsAllocationInfo))
+					{
+						CurrentTagsAllocationInfo = new FCallStackTagsAllocationInfo(0, 0);
+					}
+					TaggedAllocationInfo[TaggedAllocationInfoPair.Key] = CurrentTagsAllocationInfo + TaggedAllocationInfoPair.Value;
+				}
+			}
+		}
 
 		/// <summary>
 		/// Diffs the two passed in callstack infos and returns the difference.
 		/// </summary>
 		/// <param name="New"> Newer callstack info to subtract older from </param>
 		/// <param name="Old"> Older callstack info to subtract from older </param>
-        public static FCallStackAllocationInfo Diff( FCallStackAllocationInfo New, FCallStackAllocationInfo Old )
+		public static FCallStackAllocationInfo Diff( FCallStackAllocationInfo New, FCallStackAllocationInfo Old )
         {
             if( New.CallStackIndex != Old.CallStackIndex )
             {
                 throw new InvalidDataException();
             }
-            return new FCallStackAllocationInfo( New.Size - Old.Size, New.CallStackIndex, New.Count - Old.Count );
+
+			FCallStackAllocationInfo DiffData = New.DeepCopy();
+
+			DiffData.UntaggedAllocationInfo -= Old.UntaggedAllocationInfo;
+
+			if (Old.TaggedAllocationInfo != null)
+			{
+				DiffData.EnsureTaggedAllocationsAvailable();
+
+				foreach (var TaggedAllocationInfoPair in Old.TaggedAllocationInfo)
+				{
+					FCallStackTagsAllocationInfo CurrentTagsAllocationInfo;
+					if (!DiffData.TaggedAllocationInfo.TryGetValue(TaggedAllocationInfoPair.Key, out CurrentTagsAllocationInfo))
+					{
+						CurrentTagsAllocationInfo = new FCallStackTagsAllocationInfo(0, 0);
+					}
+					DiffData.TaggedAllocationInfo[TaggedAllocationInfoPair.Key] = CurrentTagsAllocationInfo - TaggedAllocationInfoPair.Value;
+				}
+			}
+
+			return DiffData;
         }
+
+		/// <summary>
+		/// Get a new callstack info that only contains information about allocations that match the given tags (note: this removes any tag information from the resultant object).
+		/// </summary>
+		/// <param name="InActiveTags"> Tags to filter on </param>
+		/// <param name="bInclusiveFilter"> true to include alloctions that match the tags, false to include allocations that don't match the tags </param>
+		public FCallStackAllocationInfo GetAllocationInfoForTags(ISet<FAllocationTag> InActiveTags, bool bInclusiveFilter)
+		{
+			long TagsSize;
+			int TagsCount;
+			GetSizeAndCountForTags(InActiveTags, bInclusiveFilter, out TagsSize, out TagsCount);
+			return new FCallStackAllocationInfo(TagsSize, CallStackIndex, TagsCount, -1);
+		}
+
+		public void GetSizeAndCountForTags(ISet<FAllocationTag> InActiveTags, bool bInclusiveFilter, out long OutSize, out int OutCount)
+		{
+			OutSize = 0;
+			OutCount = 0;
+
+			if (AllocationMatchesTagFilter(FStreamInfo.GlobalInstance.TagHierarchy.GetUntaggedNode().GetTag().Value, InActiveTags, bInclusiveFilter))
+			{
+				OutSize += UntaggedAllocationInfo.Size;
+				OutCount += UntaggedAllocationInfo.Count;
+			}
+
+			if (TaggedAllocationInfo != null)
+			{
+				foreach (var TaggedAllocationInfoPair in TaggedAllocationInfo)
+				{
+					if (AllocationMatchesTagFilter(TaggedAllocationInfoPair.Key, InActiveTags, bInclusiveFilter))
+					{
+						OutSize += TaggedAllocationInfoPair.Value.Size;
+						OutCount += TaggedAllocationInfoPair.Value.Count;
+					}
+				}
+			}
+		}
+
+		private bool AllocationMatchesTagFilter(FAllocationTag InTag, ISet<FAllocationTag> InActiveTags, bool bInclusiveFilter)
+		{
+			if (InActiveTags == null || InActiveTags.Count == 0)
+			{
+				// No filter, so match anything
+				return true;
+			}
+
+			return InActiveTags.Contains(InTag) ? bInclusiveFilter : !bInclusiveFilter;
+		}
+
+		private bool AllocationMatchesTagFilter(int InTagsIndex, ISet<FAllocationTag> InActiveTags, bool bInclusiveFilter)
+		{
+			if (InActiveTags == null || InActiveTags.Count == 0)
+			{
+				// No filter, so match anything
+				return true;
+			}
+
+			// Inclusive filter will pass if we match ANY tag
+			// Exclusive filter will pass if we match NO tags
+			foreach (var AllocTag in FStreamInfo.GlobalInstance.TagsArray[InTagsIndex].Tags)
+			{
+				if (InActiveTags.Contains(AllocTag))
+				{
+					return bInclusiveFilter;
+				}
+			}
+			return !bInclusiveFilter;
+		}
 	}
 
 	public enum ESliceTypesV4
@@ -156,13 +403,13 @@ namespace MemoryProfiler2
             LifetimeCallStackList = new List<FCallStackAllocationInfo>( FStreamInfo.GlobalInstance.CallStackArray.Count );
             for( int CallStackIndex=0; CallStackIndex<FStreamInfo.GlobalInstance.CallStackArray.Count; CallStackIndex++ )
             {
-                LifetimeCallStackList.Add( new FCallStackAllocationInfo( 0, CallStackIndex, 0 ) );
+                LifetimeCallStackList.Add( new FCallStackAllocationInfo( 0, CallStackIndex, 0, -1 ) );
             }
 			OverallMemorySlice = new List<FMemorySlice>();
 		}
 
 		/// <summary> Performs a deep copy of the relevant data structures. </summary>
-		public FStreamSnapshot DeepCopy( Dictionary<ulong, FCallStackAllocationInfo> PointerToPointerInfoMap )
+		public FStreamSnapshot DeepCopy( Dictionary<ulong, FLiveAllocationInfo> PointerToPointerInfoMap )
 		{
 			// Create new snapshot object.
 			FStreamSnapshot Snapshot = new FStreamSnapshot( "Copy" );
@@ -170,8 +417,7 @@ namespace MemoryProfiler2
 			// Manually perform a deep copy of LifetimeCallstackList
 			foreach( FCallStackAllocationInfo AllocationInfo in LifetimeCallStackList )
 			{
-				Snapshot.LifetimeCallStackList[ AllocationInfo.CallStackIndex ]
-					= new FCallStackAllocationInfo( AllocationInfo.Size, AllocationInfo.CallStackIndex, AllocationInfo.Count );
+				Snapshot.LifetimeCallStackList[ AllocationInfo.CallStackIndex ] = AllocationInfo.DeepCopy();
 			}
 
 			Snapshot.AllocationCount = AllocationCount;
@@ -184,35 +430,29 @@ namespace MemoryProfiler2
 			return Snapshot;
 		}
 
-		public void FillActiveCallStackList( Dictionary<ulong, FCallStackAllocationInfo> PointerToPointerInfoMap )
+		public void FillActiveCallStackList( Dictionary<ulong, FLiveAllocationInfo> PointerToPointerInfoMap )
 		{
 			ActiveCallStackList.Clear();
 			ActiveCallStackList.Capacity = LifetimeCallStackList.Count;
 
-			foreach( KeyValuePair<ulong, FCallStackAllocationInfo> PointerData in PointerToPointerInfoMap )
+			foreach( var PointerData in PointerToPointerInfoMap )
 			{
-				int CallStackIndex = PointerData.Value.CallStackIndex;
-				long Size = PointerData.Value.Size;
-
 				// make sure allocationInfoList is big enough
-				while( CallStackIndex >= ActiveCallStackList.Count )
+				while(PointerData.Value.CallStackIndex >= ActiveCallStackList.Count )
 				{
-					ActiveCallStackList.Add( new FCallStackAllocationInfo( 0, ActiveCallStackList.Count, 0 ) );
+					ActiveCallStackList.Add( new FCallStackAllocationInfo( 0, ActiveCallStackList.Count, 0, -1 ) );
 				}
 
-				FCallStackAllocationInfo NewAllocInfo = ActiveCallStackList[ CallStackIndex ];
-				NewAllocInfo.Size += Size;
-				NewAllocInfo.Count++;
-				ActiveCallStackList[ CallStackIndex ] = NewAllocInfo;
+				ActiveCallStackList[PointerData.Value.CallStackIndex].Add(PointerData.Value.Size, 1, PointerData.Value.TagsIndex);
 			}
 
 			// strip out any callstacks with no allocations
-			ActiveCallStackList.RemoveAll(Item => Item.Count == 0);
+			ActiveCallStackList.RemoveAll(Item => Item.TotalCount == 0);
 			ActiveCallStackList.TrimExcess();
 		}
 
 		/// <summary> Convert "callstack to allocation" mapping (either passed in or generated from pointer map) to callstack info list. </summary>
-		public void FinalizeSnapshot( Dictionary<ulong, FCallStackAllocationInfo> PointerToPointerInfoMap )
+		public void FinalizeSnapshot( Dictionary<ulong, FLiveAllocationInfo> PointerToPointerInfoMap )
 		{
 			FillActiveCallStackList( PointerToPointerInfoMap );
 		}
@@ -261,12 +501,12 @@ namespace MemoryProfiler2
 
 					if( OldAllocInfo.CallStackIndex == NewAllocInfo.CallStackIndex )
 					{
-						long ResultSize = NewAllocInfo.Size - OldAllocInfo.Size;
-						int ResultCount = NewAllocInfo.Count - OldAllocInfo.Count;
+						long ResultSize = NewAllocInfo.TotalSize - OldAllocInfo.TotalSize;
+						int ResultCount = NewAllocInfo.TotalCount - OldAllocInfo.TotalCount;
 
 						if( ResultSize != 0 || ResultCount != 0 )
 						{
-							ResultActiveCallStackList.Add( new FCallStackAllocationInfo( ResultSize, NewAllocInfo.CallStackIndex, ResultCount ) );
+							ResultActiveCallStackList.Add( new FCallStackAllocationInfo( ResultSize, NewAllocInfo.CallStackIndex, ResultCount, -1 ) );
 						}
 
 						OldIndex++;
@@ -279,7 +519,7 @@ namespace MemoryProfiler2
 					}
 					else // OldAllocInfo.CallStackIndex < NewAllocInfo.CallStackIndex
 					{
-						ResultActiveCallStackList.Add( new FCallStackAllocationInfo( -OldAllocInfo.Size, OldAllocInfo.CallStackIndex, -OldAllocInfo.Count ) );
+						ResultActiveCallStackList.Add( new FCallStackAllocationInfo( -OldAllocInfo.TotalSize, OldAllocInfo.CallStackIndex, -OldAllocInfo.TotalCount, -1 ) );
 						OldIndex++;
 					}
 
@@ -362,10 +602,10 @@ namespace MemoryProfiler2
 			foreach( FCallStackAllocationInfo AllocationInfo in CallStackList )
 			{
                 // Skip callstacks with no contribution in this snapshot.
-                if( AllocationInfo.Count > 0 )
+                if( AllocationInfo.TotalCount > 0 )
                 {
 				    // Dump size first, followed by count.
-				    CSVWriter.Write(AllocationInfo.Size + "," + AllocationInfo.Count + ",");
+				    CSVWriter.Write(AllocationInfo.TotalSize + "," + AllocationInfo.TotalCount + ",");
 
 				    // Iterate over ach address in callstack and dump to CSV.
 				    FCallStack CallStack = FStreamInfo.GlobalInstance.CallStackArray[AllocationInfo.CallStackIndex];
diff --git a/Engine/Source/Programs/MemoryProfiler2/StreamToken.cs b/Engine/Source/Programs/MemoryProfiler2/StreamToken.cs
--- a/Engine/Source/Programs/MemoryProfiler2/StreamToken.cs
+++ b/Engine/Source/Programs/MemoryProfiler2/StreamToken.cs
@@ -109,11 +109,6 @@ namespace MemoryProfiler2
 		SUBTYPE_Unknown,
 	};
 
-	public interface IMemoryStats
-	{
-
-	}
-
 	/// <summary> 
 	/// Struct used to hold platform memory stats and allocator stats. Valid only for MemoryProfiler2.FProfileDataHeader.Version >= 4.
 	/// @see FGenericPlatformMemory::GetStatsForMallocProfiler 
@@ -238,7 +233,7 @@ namespace MemoryProfiler2
 	/// Variable sized token emitted by capture code. The parsing code ReadNextToken deals with this and updates
 	/// internal data. The calling code is responsible for only accessing member variables associated with the type.
 	/// </summary>
-	public class FStreamToken//@TODO: Can I turn this back into a struct?
+	public class FStreamToken
 	{
 		// Parsing configuration
 
@@ -275,6 +270,9 @@ namespace MemoryProfiler2
 		/// <summary> Index into callstack array. </summary>
 		public Int32 CallStackIndex;
 
+		/// <summary> Index into tags array. </summary>
+		public Int32 TagsIndex;
+
 		/// <summary> Size of allocation in alloc / realloc case. </summary>
 		public Int32 Size;
 
@@ -423,6 +421,7 @@ namespace MemoryProfiler2
 			Metrics.Clear();
 			LoadedLevels.Clear();
 			CallStackIndex = -1;
+			TagsIndex = -1;
 			ScriptCallstackIndex = -1;
 			ScriptObjectTypeIndex = -1;
 
@@ -441,6 +440,12 @@ namespace MemoryProfiler2
 					// Get the call stack index.
 					CallStackIndex = BinaryStream.ReadInt32();
 
+					// Get the tags index.
+					if (Version >= 7)
+					{
+						TagsIndex = BinaryStream.ReadInt32();
+					}
+
 					// Get the size of an allocation.
 					UInt32 UnsignedSize = BinaryStream.ReadUInt32();
 
@@ -470,6 +475,12 @@ namespace MemoryProfiler2
 					NewPointer = BinaryStream.ReadUInt64();
 					CallStackIndex = BinaryStream.ReadInt32();
 
+					// Get the tags index.
+					if (Version >= 7)
+					{
+						TagsIndex = BinaryStream.ReadInt32();
+					}
+
 					UInt32 UnsignedSize = BinaryStream.ReadUInt32();
 					bool bHasGCMData = ReadGCMData( BinaryStream, ref UnsignedSize );
 					if( bHasGCMData == false )
diff --git a/Engine/Source/Programs/MemoryProfiler2/TreeViewEx.cs b/Engine/Source/Programs/MemoryProfiler2/TreeViewEx.cs
--- a/Engine/Source/Programs/MemoryProfiler2/TreeViewEx.cs
+++ b/Engine/Source/Programs/MemoryProfiler2/TreeViewEx.cs
@@ -26,20 +26,58 @@ namespace MemoryProfiler2
 			TreeViewWin32.SendMessage( this.Handle, TreeViewWin32.TVM_SETEXTENDEDSTYLE, (IntPtr)TreeViewWin32.TVS_EX_DOUBLEBUFFER, (IntPtr)TreeViewWin32.TVS_EX_DOUBLEBUFFER );
 			base.OnHandleCreated( e );
 		}
+
+		public void SetCheckBoxImageState(TreeNode InTreeNode, UInt32 InImageStateIndex)
+		{
+			var TreeViewItem = new TreeViewWin32.TVITEM();
+			TreeViewItem.hItem = InTreeNode.Handle;
+			TreeViewItem.mask = TreeViewWin32.TVIF_STATE;
+			TreeViewItem.stateMask = TreeViewWin32.TVIS_STATEIMAGEMASK;
+			TreeViewItem.state = InImageStateIndex << 12; // see INDEXTOSTATEIMAGEMASK
+			TreeViewWin32.SendMessage( this.Handle, TreeViewWin32.TVM_SETITEM, IntPtr.Zero, ref TreeViewItem );
+		}
 	}
 
 	/// <summary> Helper class with native Win32 functions, constants and structures. </summary>
 	public class TreeViewWin32
 	{
+		[StructLayout(LayoutKind.Sequential, Pack = 8, CharSet = CharSet.Auto)]
+		public struct TVITEM
+		{
+			public UInt32 mask;
+			public IntPtr hItem;
+			public UInt32 state;
+			public UInt32 stateMask;
+			[MarshalAs(UnmanagedType.LPTStr)]
+			public string lpszText;
+			public int cchTextMax;
+			public int iImage;
+			public int iSelectedImage;
+			public int cChildren;
+			public IntPtr lParam;
+		}
+
 		[DllImport( "user32.dll", EntryPoint = "SendMessage" )]
 		public static extern IntPtr SendMessage( IntPtr hwnd, UInt32 wMsg, IntPtr wParam, IntPtr lParam );
 
+		[DllImport( "user32.dll", EntryPoint = "SendMessage", CharSet = CharSet.Auto )]
+		public static extern IntPtr SendMessage( IntPtr hWnd, UInt32 Msg, IntPtr wParam, ref TVITEM lParam );
+
 		public const UInt32 TV_FIRST = 0x1100;
 
 		/// <summary> Informs the tree-view control to set extended styles. </summary>
 		public const UInt32 TVM_SETEXTENDEDSTYLE = TV_FIRST + 44;
 
+		/// <summary> Sets some or all of a tree-view item's attributes. </summary>
+		public const UInt32 TVM_SETITEM = TV_FIRST + 63;
+
 		/// <summary> Specifies how the background is erased or filled. </summary>
 		public const UInt32 TVS_EX_DOUBLEBUFFER = 0x0004;
+
+		/// <summary> The state and stateMask members are valid. </summary>
+		public const UInt32 TVIF_STATE = 0x8;
+
+		/// <summary> Mask for the bits used to specify the item's state image index. </summary>
+		public const UInt32 TVIS_STATEIMAGEMASK = 0xF000;
 	}
 }
\ No newline at end of file
diff --git a/Engine/Source/Programs/UnrealBuildTool/Configuration/UEBuildTarget.cs b/Engine/Source/Programs/UnrealBuildTool/Configuration/UEBuildTarget.cs
--- a/Engine/Source/Programs/UnrealBuildTool/Configuration/UEBuildTarget.cs
+++ b/Engine/Source/Programs/UnrealBuildTool/Configuration/UEBuildTarget.cs
@@ -2056,13 +2056,17 @@ namespace UnrealBuildTool
             if (Rules.bUsesSlate)
             {
 				Receipt.RuntimeDependencies.Add("$(EngineDir)/Content/Slate/...", StagedFileType.UFS);
-				if(ProjectFile != null)
+				if (Configuration != UnrealTargetConfiguration.Shipping)
 				{
-					Receipt.RuntimeDependencies.Add("$(ProjectDir)/Content/Slate/...", StagedFileType.UFS);
+					Receipt.RuntimeDependencies.Add("$(EngineDir)/Content/SlateDebug/...", StagedFileType.UFS);
 				}
-				if (Rules.bUsesSlateEditorStyle)
+				if (ProjectFile != null)
 				{
-					Receipt.RuntimeDependencies.Add("$(EngineDir)/Content/Editor/Slate/...", StagedFileType.UFS);
+					Receipt.RuntimeDependencies.Add("$(ProjectDir)/Content/Slate/...", StagedFileType.UFS);
+					if (Configuration != UnrealTargetConfiguration.Shipping)
+					{
+						Receipt.RuntimeDependencies.Add("$(EngineDir)/Content/SlateDebug/...", StagedFileType.UFS);
+					}
 				}
 			}
 
diff --git a/Engine/Source/Programs/UnrealBuildTool/System/RulesCompiler.cs b/Engine/Source/Programs/UnrealBuildTool/System/RulesCompiler.cs
--- a/Engine/Source/Programs/UnrealBuildTool/System/RulesCompiler.cs
+++ b/Engine/Source/Programs/UnrealBuildTool/System/RulesCompiler.cs
@@ -610,12 +610,6 @@ namespace UnrealBuildTool
 		/// </summary>
 		public bool bUsesSlate = true;
 
-		/// <summary>
-		/// Hack for legacy game styling isses.  No new project should ever set this to true
-		/// Whether the project uses the Slate editor style in game.  
-		/// </summary>
-		public bool bUsesSlateEditorStyle = false;
-
 		/// <summary>
 		/// Forces linking against the static CRT. This is not supported across the engine due to the need for allocator implementations to be shared (for example), and TPS 
 		/// libraries to be consistent with each other, but can be used for utility programs.
diff --git a/Engine/Source/Programs/UnrealBuildTool/Windows/UEBuildWindows.cs b/Engine/Source/Programs/UnrealBuildTool/Windows/UEBuildWindows.cs
--- a/Engine/Source/Programs/UnrealBuildTool/Windows/UEBuildWindows.cs
+++ b/Engine/Source/Programs/UnrealBuildTool/Windows/UEBuildWindows.cs
@@ -245,9 +245,9 @@ namespace UnrealBuildTool
 				}
 				else
 				{
-					// Windows Vista or higher required
-					InBuildTarget.GlobalCompileEnvironment.Config.Definitions.Add("_WIN32_WINNT=0x0600");
-					InBuildTarget.GlobalCompileEnvironment.Config.Definitions.Add("WINVER=0x0600");
+					// Windows 7 or higher required
+					InBuildTarget.GlobalCompileEnvironment.Config.Definitions.Add("_WIN32_WINNT=0x0601");
+					InBuildTarget.GlobalCompileEnvironment.Config.Definitions.Add("WINVER=0x0601");
 				}
 			}
 			InBuildTarget.GlobalCompileEnvironment.Config.Definitions.Add("PLATFORM_WINDOWS=1");
diff --git a/Engine/Source/Runtime/AutomationMessages/Classes/AutomationWorkerMessages.h b/Engine/Source/Runtime/AutomationMessages/Classes/AutomationWorkerMessages.h
--- a/Engine/Source/Runtime/AutomationMessages/Classes/AutomationWorkerMessages.h
+++ b/Engine/Source/Runtime/AutomationMessages/Classes/AutomationWorkerMessages.h
@@ -170,9 +170,35 @@ struct FAutomationWorkerRequestTestsReply
 {
 	GENERATED_USTRUCT_BODY()
 
-	/** Holds the test information serialized into a string. */
 	UPROPERTY(EditAnywhere, Category="Message")
-	FString TestInfo;
+	FString DisplayName;
+
+	UPROPERTY(EditAnywhere, Category="Message")
+	FString FullTestPath;
+
+	UPROPERTY(EditAnywhere, Category="Message")
+	FString TestName;
+
+	UPROPERTY(EditAnywhere, Category="Message")
+	FString TestParameter;
+
+	UPROPERTY(EditAnywhere, Category="Message")
+	FString SourceFile;
+
+	UPROPERTY(EditAnywhere, Category="Message")
+	int32 SourceFileLine;
+
+	UPROPERTY(EditAnywhere, Category="Message")
+	FString AssetPath;
+
+	UPROPERTY(EditAnywhere, Category="Message")
+	FString OpenCommand;
+
+	UPROPERTY(EditAnywhere, Category="Message")
+	uint32 TestFlags;
+
+	UPROPERTY(EditAnywhere, Category="Message")
+	uint32 NumParticipantsRequired;
 
 	/** Holds the total number of tests returned. */
 	UPROPERTY(EditAnywhere, Category="Message")
@@ -182,10 +208,35 @@ struct FAutomationWorkerRequestTestsReply
 	FAutomationWorkerRequestTestsReply() { }
 
 	/** Creates and initializes a new instance. */
-	FAutomationWorkerRequestTestsReply(const FString& InTestInfo, const int32& InTotalNumTests)
-		: TestInfo(InTestInfo)
-		, TotalNumTests(InTotalNumTests)
-	{ }
+	FAutomationWorkerRequestTestsReply(const FAutomationTestInfo& InTestInfo, const int32& InTotalNumTests)
+		: TotalNumTests(InTotalNumTests)
+	{
+		DisplayName = InTestInfo.GetDisplayName();
+		FullTestPath = InTestInfo.GetFullTestPath();
+		TestName = InTestInfo.GetTestName();
+		TestParameter = InTestInfo.GetTestParameter();
+		SourceFile = InTestInfo.GetSourceFile();
+		SourceFileLine = InTestInfo.GetSourceFileLine();
+		AssetPath = InTestInfo.GetAssetPath();
+		OpenCommand = InTestInfo.GetOpenCommand();
+		TestFlags = InTestInfo.GetTestFlags();
+		NumParticipantsRequired = InTestInfo.GetNumParticipantsRequired();
+	}
+
+	FAutomationTestInfo GetTestInfo() const
+	{
+		return FAutomationTestInfo(
+			DisplayName,
+			FullTestPath,
+			TestName,
+			TestFlags,
+			NumParticipantsRequired,
+			TestParameter,
+			SourceFile,
+			SourceFileLine,
+			AssetPath,
+			OpenCommand);
+	}
 };
 
 
