MemoryProfiler2: Search for source files in platform extension directories.

diff --git a/Engine/Source/Programs/MemoryProfiler2/MemoryProfiler2.csproj b/Engine/Source/Programs/MemoryProfiler2/MemoryProfiler2.csproj
--- a/Engine/Source/Programs/MemoryProfiler2/MemoryProfiler2.csproj
+++ b/Engine/Source/Programs/MemoryProfiler2/MemoryProfiler2.csproj
@@ -189,20 +189,6 @@
       <DesignTimeSharedInput>True</DesignTimeSharedInput>
     </Compile>
   </ItemGroup>
-  <Choose>
-    <When Condition="Exists('PS4\PS4SymbolParser.cs')">
-      <ItemGroup>
-        <Compile Include="PS4\PS4SymbolParser.cs" />
-      </ItemGroup>
-    </When>
-  </Choose>
-  <Choose>
-    <When Condition="Exists('Switch\SwitchSymbolParser.cs')">
-      <ItemGroup>
-        <Compile Include="Switch\SwitchSymbolParser.cs" />
-      </ItemGroup>
-    </When>
-  </Choose>
   <ItemGroup>
     <Compile Include="ListViewEx.cs">
       <SubType>Component</SubType>
@@ -224,6 +210,12 @@
     <Compile Include="StreamSnapshot.cs" />
     <Compile Include="StreamToken.cs" />
     <Compile Include="TimeLineView.cs" />
+    <Compile Include="../../../Platforms/*/Source/Programs/MemoryProfiler2/**/*.cs">
+      <!-- RecursiveDir metadata expands to [PlatformName]/Source/Programs/[etc]. The Replace() will
+      replace everything from the first slash to the end leaving just the platform name. Should the
+      Replace() fail then RecursiveDir is still an agreeable value -->
+      <Link>$([System.Text.RegularExpressions.Regex]::Replace(%(Compile.RecursiveDir), [\\/].+$, ``))/_</Link>
+    </Compile>
   </ItemGroup>
   <ItemGroup>
     <BootstrapperPackage Include="Microsoft.Net.Client.3.5">
