Update the MemoryProfiler2 project to .NET 4.8 to fix errors when loading the project in VS2022

diff --git a/Engine/Source/Programs/MemoryProfiler2/App.config b/Engine/Source/Programs/MemoryProfiler2/App.config
--- a/Engine/Source/Programs/MemoryProfiler2/App.config
+++ b/Engine/Source/Programs/MemoryProfiler2/App.config
@@ -1,10 +1,10 @@
 <?xml version="1.0"?>
 <configuration>
 	<startup>
-		<supportedRuntime version="v4.0" sku=".NETFramework,Version=v4.5"/>
+		<supportedRuntime version="v4.0" sku=".NETFramework,Version=v4.8"/>
 	</startup>
 	<runtime>
 		<generatePublisherEvidence enabled="false"/>
-        <gcAllowVeryLargeObjects enabled="true" />
+        <gcAllowVeryLargeObjects enabled="true"/>
 	</runtime>
 </configuration>
diff --git a/Engine/Source/Programs/MemoryProfiler2/MemoryProfiler2.csproj b/Engine/Source/Programs/MemoryProfiler2/MemoryProfiler2.csproj
--- a/Engine/Source/Programs/MemoryProfiler2/MemoryProfiler2.csproj
+++ b/Engine/Source/Programs/MemoryProfiler2/MemoryProfiler2.csproj
@@ -1,5 +1,5 @@
 ﻿<?xml version="1.0" encoding="utf-8"?>
-<Project DefaultTargets="Build" xmlns="http://schemas.microsoft.com/developer/msbuild/2003" ToolsVersion="4.0">
+<Project DefaultTargets="Build" xmlns="http://schemas.microsoft.com/developer/msbuild/2003" ToolsVersion="12.0">
   <PropertyGroup>
     <Configuration Condition=" '$(Configuration)' == '' ">Debug</Configuration>
     <Platform Condition=" '$(Platform)' == '' ">AnyCPU</Platform>
@@ -23,7 +23,7 @@
     <OldToolsVersion>3.5</OldToolsVersion>
     <UpgradeBackupLocation>
     </UpgradeBackupLocation>
-    <TargetFrameworkVersion>v4.5</TargetFrameworkVersion>
+    <TargetFrameworkVersion>v4.8</TargetFrameworkVersion>
     <IsWebBootstrapper>false</IsWebBootstrapper>
     <TargetFrameworkProfile>
     </TargetFrameworkProfile>
@@ -264,4 +264,4 @@
   <Target Name="AfterBuild">
   </Target>
   -->
-</Project>
+</Project>
\ No newline at end of file
