Change MemoryProfiler2's targeted .NET version to 4.6.2 instead of 4.8 to fix a CIS warning for the project.

diff --git a/Engine/Source/Programs/MemoryProfiler2/App.config b/Engine/Source/Programs/MemoryProfiler2/App.config
--- a/Engine/Source/Programs/MemoryProfiler2/App.config
+++ b/Engine/Source/Programs/MemoryProfiler2/App.config
@@ -1,7 +1,7 @@
 <?xml version="1.0"?>
 <configuration>
 	<startup>
-		<supportedRuntime version="v4.0" sku=".NETFramework,Version=v4.8"/>
+		<supportedRuntime version="v4.0" sku=".NETFramework,Version=v4.6.2"/>
 	</startup>
 	<runtime>
 		<generatePublisherEvidence enabled="false"/>
diff --git a/Engine/Source/Programs/MemoryProfiler2/MemoryProfiler2.csproj b/Engine/Source/Programs/MemoryProfiler2/MemoryProfiler2.csproj
--- a/Engine/Source/Programs/MemoryProfiler2/MemoryProfiler2.csproj
+++ b/Engine/Source/Programs/MemoryProfiler2/MemoryProfiler2.csproj
@@ -23,7 +23,7 @@
     <OldToolsVersion>3.5</OldToolsVersion>
     <UpgradeBackupLocation>
     </UpgradeBackupLocation>
-    <TargetFrameworkVersion>v4.8</TargetFrameworkVersion>
+    <TargetFrameworkVersion>v4.6.2</TargetFrameworkVersion>
     <IsWebBootstrapper>false</IsWebBootstrapper>
     <TargetFrameworkProfile>
     </TargetFrameworkProfile>
@@ -264,4 +264,4 @@
   <Target Name="AfterBuild">
   </Target>
   -->
-</Project>
\ No newline at end of file
+</Project>
