Remove unacceptable string from MemoryProfiler2

#rnx

diff --git a/Engine/Source/Programs/MemoryProfiler2/StreamParser.cs b/Engine/Source/Programs/MemoryProfiler2/StreamParser.cs
--- a/Engine/Source/Programs/MemoryProfiler2/StreamParser.cs
+++ b/Engine/Source/Programs/MemoryProfiler2/StreamParser.cs
@@ -576,7 +576,7 @@ namespace MemoryProfiler2
 									NextDataFile++;
 									if( NextDataFile > Header.NumDataFiles )
 									{
-										throw new InvalidDataException( "Found an unexpected number of data files (more than indicated in the master file" );
+										throw new InvalidDataException( "Found an unexpected number of data files (more than indicated in the primary file" );
 									}
 									break;
 								}
