[MemoryProfiler2]
* Upgraded from .net framework 4.6 to .net 6

diff --git a/Engine/Source/Programs/MemoryProfiler2/CallstackHistoryView.cs b/Engine/Source/Programs/MemoryProfiler2/CallstackHistoryView.cs
--- a/Engine/Source/Programs/MemoryProfiler2/CallstackHistoryView.cs
+++ b/Engine/Source/Programs/MemoryProfiler2/CallstackHistoryView.cs
@@ -1,4 +1,4 @@
-﻿// Copyright Epic Games, Inc. All Rights Reserved.
+// Copyright Epic Games, Inc. All Rights Reserved.
 
 using System;
 using System.Collections.Generic;
@@ -683,7 +683,7 @@ namespace MemoryProfiler2
 					}
 					else
 					{
-						throw e;
+						throw;
 					}
 				}
 
diff --git a/Engine/Source/Programs/MemoryProfiler2/FilterTagsDialog.cs b/Engine/Source/Programs/MemoryProfiler2/FilterTagsDialog.cs
--- a/Engine/Source/Programs/MemoryProfiler2/FilterTagsDialog.cs
+++ b/Engine/Source/Programs/MemoryProfiler2/FilterTagsDialog.cs
@@ -8,6 +8,7 @@ using System.Drawing;
 using System.Linq;
 using System.Text;
 using System.Windows.Forms;
+using Timer = System.Windows.Forms.Timer;
 
 namespace MemoryProfiler2
 {
@@ -58,7 +59,7 @@ namespace MemoryProfiler2
 			}
 
 			// Build the WinForm tree
-			if (FilterTagsTreeView.Handle != null) // This if statement exists to force the native tree view handle to be created
+			if (FilterTagsTreeView.Handle != IntPtr.Zero) // This if statement exists to force the native tree view handle to be created
 			{
 				bool bIsFiltering = InFilterTerms != null && InFilterTerms.Length > 0;
 
diff --git a/Engine/Source/Programs/MemoryProfiler2/MainWindow.cs b/Engine/Source/Programs/MemoryProfiler2/MainWindow.cs
--- a/Engine/Source/Programs/MemoryProfiler2/MainWindow.cs
+++ b/Engine/Source/Programs/MemoryProfiler2/MainWindow.cs
@@ -21,6 +21,9 @@ using EpicCommonUtils;
 using UnrealControls;
 using System.Drawing.Drawing2D;
 using System.Windows.Forms.DataVisualization.Charting;
+using System.Runtime.InteropServices;
+using System.Runtime.Versioning;
+using System.Security;
 
 namespace MemoryProfiler2
 {
@@ -609,15 +612,16 @@ namespace MemoryProfiler2
 		protected EnvDTE.DTE FindOrCreateDTEInstance( string ObjectName )
 		{
 			// Try to find an open copy, and if that fails, create one
-			EnvDTE.DTE result;
+			EnvDTE.DTE result = null;
 			try
 			{
-				result = System.Runtime.InteropServices.Marshal.GetActiveObject( ObjectName ) as EnvDTE.DTE;
+				result = Marshal2.GetActiveObject( ObjectName ) as EnvDTE.DTE;
 			}
 			catch
 			{
 				System.Type t = System.Type.GetTypeFromProgID( ObjectName );
-				result = System.Activator.CreateInstance( t ) as EnvDTE.DTE;
+				if (t != null)
+					result = System.Activator.CreateInstance( t ) as EnvDTE.DTE;
 			}
 
 			return result;
@@ -2761,4 +2765,54 @@ namespace MemoryProfiler2
 			return ResultArray;
 		}
 	}
+
+	public static class Marshal2
+	{
+		internal const String OLEAUT32 = "oleaut32.dll";
+		internal const String OLE32 = "ole32.dll";
+
+		[System.Security.SecurityCritical]  // auto-generated_required
+		public static Object GetActiveObject(String progID)
+		{
+			Object obj = null;
+			Guid clsid;
+
+			// Call CLSIDFromProgIDEx first then fall back on CLSIDFromProgID if
+			// CLSIDFromProgIDEx doesn't exist.
+			try
+			{
+				CLSIDFromProgIDEx(progID, out clsid);
+			}
+			//            catch
+			catch (Exception)
+			{
+				CLSIDFromProgID(progID, out clsid);
+			}
+
+			GetActiveObject(ref clsid, IntPtr.Zero, out obj);
+			return obj;
+		}
+
+		//[DllImport(Microsoft.Win32.Win32Native.OLE32, PreserveSig = false)]
+		[DllImport(OLE32, PreserveSig = false)]
+		[ResourceExposure(ResourceScope.None)]
+		[SuppressUnmanagedCodeSecurity]
+		[System.Security.SecurityCritical]  // auto-generated
+		private static extern void CLSIDFromProgIDEx([MarshalAs(UnmanagedType.LPWStr)] String progId, out Guid clsid);
+
+		//[DllImport(Microsoft.Win32.Win32Native.OLE32, PreserveSig = false)]
+		[DllImport(OLE32, PreserveSig = false)]
+		[ResourceExposure(ResourceScope.None)]
+		[SuppressUnmanagedCodeSecurity]
+		[System.Security.SecurityCritical]  // auto-generated
+		private static extern void CLSIDFromProgID([MarshalAs(UnmanagedType.LPWStr)] String progId, out Guid clsid);
+
+		//[DllImport(Microsoft.Win32.Win32Native.OLEAUT32, PreserveSig = false)]
+		[DllImport(OLEAUT32, PreserveSig = false)]
+		[ResourceExposure(ResourceScope.None)]
+		[SuppressUnmanagedCodeSecurity]
+		[System.Security.SecurityCritical]  // auto-generated
+		private static extern void GetActiveObject(ref Guid rclsid, IntPtr reserved, [MarshalAs(UnmanagedType.Interface)] out Object ppunk);
+
+	}
 }
diff --git a/Engine/Source/Programs/MemoryProfiler2/MemoryProfiler2.csproj b/Engine/Source/Programs/MemoryProfiler2/MemoryProfiler2.csproj
--- a/Engine/Source/Programs/MemoryProfiler2/MemoryProfiler2.csproj
+++ b/Engine/Source/Programs/MemoryProfiler2/MemoryProfiler2.csproj
@@ -1,267 +1,25 @@
-﻿<?xml version="1.0" encoding="utf-8"?>
-<Project DefaultTargets="Build" xmlns="http://schemas.microsoft.com/developer/msbuild/2003" ToolsVersion="12.0">
+﻿<Project Sdk="Microsoft.NET.Sdk">
+
   <PropertyGroup>
-    <Configuration Condition=" '$(Configuration)' == '' ">Debug</Configuration>
-    <Platform Condition=" '$(Platform)' == '' ">AnyCPU</Platform>
-    <ProductVersion>9.0.30729</ProductVersion>
-    <SchemaVersion>2.0</SchemaVersion>
-    <ProjectGuid>{75240A67-CF47-4DDC-83B2-AD5D371F31BD}</ProjectGuid>
     <OutputType>WinExe</OutputType>
-    <AppDesignerFolder>Properties</AppDesignerFolder>
-    <RootNamespace>MemoryProfiler2</RootNamespace>
-    <AssemblyName>MemoryProfiler2</AssemblyName>
-    <SccProjectName>
-    </SccProjectName>
-    <SccLocalPath>
-    </SccLocalPath>
-    <SccAuxPath>
-    </SccAuxPath>
-    <SccProvider>
-    </SccProvider>
-    <FileUpgradeFlags>
-    </FileUpgradeFlags>
-    <OldToolsVersion>3.5</OldToolsVersion>
-    <UpgradeBackupLocation>
-    </UpgradeBackupLocation>
-    <TargetFrameworkVersion>v4.6.2</TargetFrameworkVersion>
-    <IsWebBootstrapper>false</IsWebBootstrapper>
-    <TargetFrameworkProfile>
-    </TargetFrameworkProfile>
-    <IsWebBootstrapper>false</IsWebBootstrapper>
-    <PublishUrl>publish\</PublishUrl>
-    <Install>true</Install>
-    <InstallFrom>Disk</InstallFrom>
-    <UpdateEnabled>false</UpdateEnabled>
-    <UpdateMode>Foreground</UpdateMode>
-    <UpdateInterval>7</UpdateInterval>
-    <UpdateIntervalUnits>Days</UpdateIntervalUnits>
-    <UpdatePeriodically>false</UpdatePeriodically>
-    <UpdateRequired>false</UpdateRequired>
-    <MapFileExtensions>true</MapFileExtensions>
-    <ApplicationRevision>0</ApplicationRevision>
-    <ApplicationVersion>1.0.0.%2a</ApplicationVersion>
-    <UseApplicationTrust>false</UseApplicationTrust>
-    <BootstrapperEnabled>true</BootstrapperEnabled>
-  </PropertyGroup>
-  <PropertyGroup Condition=" '$(Configuration)|$(Platform)' == 'Debug|AnyCPU' ">
-    <DebugSymbols>true</DebugSymbols>
-    <DebugType>full</DebugType>
-    <Optimize>false</Optimize>
-    <OutputPath>..\..\..\Programs\MemoryProfiler2\Binaries\</OutputPath>
-    <DefineConstants>TRACE;DEBUG</DefineConstants>
-    <ErrorReport>prompt</ErrorReport>
-    <WarningLevel>4</WarningLevel>
-    <PlatformTarget>AnyCPU</PlatformTarget>
-    <Prefer32Bit>false</Prefer32Bit>
-  </PropertyGroup>
-  <PropertyGroup Condition=" '$(Configuration)|$(Platform)' == 'Release|AnyCPU' ">
-    <DebugType>pdbonly</DebugType>
-    <Optimize>true</Optimize>
-    <OutputPath>..\..\..\Programs\MemoryProfiler2\Binaries\</OutputPath>
-    <DefineConstants>TRACE</DefineConstants>
-    <ErrorReport>prompt</ErrorReport>
-    <WarningLevel>4</WarningLevel>
-    <PlatformTarget>AnyCPU</PlatformTarget>
-    <Prefer32Bit>false</Prefer32Bit>
-  </PropertyGroup>
-  <PropertyGroup Condition=" '$(Configuration)|$(Platform)' == 'Development|AnyCPU' ">
-    <DebugType>pdbonly</DebugType>
-    <Optimize>true</Optimize>
-    <OutputPath>..\..\..\Programs\MemoryProfiler2\Binaries\</OutputPath>
-    <DefineConstants>TRACE</DefineConstants>
-    <ErrorReport>prompt</ErrorReport>
-    <WarningLevel>4</WarningLevel>
-    <PlatformTarget>AnyCPU</PlatformTarget>
-    <Prefer32Bit>false</Prefer32Bit>
-  </PropertyGroup>
-  <PropertyGroup>
-    <ApplicationIcon>Icon.ico</ApplicationIcon>
+    <TargetFramework>net6.0-windows</TargetFramework>
+    <Nullable>disable</Nullable>
+    <UseWindowsForms>true</UseWindowsForms>
+    <ImplicitUsings>enable</ImplicitUsings>
+    <ProduceReferenceAssembly>True</ProduceReferenceAssembly>
+    <PackageId />
+    <AssemblyName></AssemblyName>
   </PropertyGroup>
+
+  <ItemGroup>
+    <PackageReference Include="envdte" Version="17.3.32804.24" />
+    <PackageReference Include="HIC.System.Windows.Forms.DataVisualization" Version="1.0.1" />
+  </ItemGroup>
+
   <ItemGroup>
-    <Reference Include="EnvDTE, Version=8.0.0.0, Culture=neutral, PublicKeyToken=b03f5f7f11d50a3a" />
-    <Reference Include="System" />
-    <Reference Include="System.Core">
-      <RequiredTargetFramework>3.5</RequiredTargetFramework>
-    </Reference>
-    <Reference Include="System.Data" />
-    <Reference Include="System.Data.Linq">
-      <RequiredTargetFramework>3.5</RequiredTargetFramework>
-    </Reference>
-    <Reference Include="System.Deployment" />
-    <Reference Include="System.Drawing" />
-    <Reference Include="System.Windows.Forms" />
-    <Reference Include="System.Windows.Forms.DataVisualization" />
-    <Reference Include="System.Xml" />
     <Reference Include="UnrealControls">
       <HintPath>..\..\..\Binaries\DotNET\UnrealControls.dll</HintPath>
     </Reference>
   </ItemGroup>
-  <ItemGroup>
-    <Compile Include="AllocationTags.cs" />
-    <Compile Include="BinaryReaderBigEndian.cs" />
-    <Compile Include="CallGraphTreeViewParser.cs" />
-    <Compile Include="CallStack.cs" />
-    <Compile Include="CallStackAddress.cs" />
-    <Compile Include="CallStackGroup.cs" />
-    <Compile Include="CallstackHistoryView.cs" />
-    <Compile Include="CommandLineWrapper.cs" />
-    <Compile Include="DoubleBufferedPanel.cs">
-      <SubType>Component</SubType>
-    </Compile>
-    <Compile Include="ElfSymbolParser.cs" />
-    <Compile Include="ExclusiveListViewParser.cs" />
-    <Compile Include="FilterTagsDialog.cs">
-      <SubType>Form</SubType>
-    </Compile>
-    <Compile Include="FilterTagsDialog.Designer.cs">
-      <DependentUpon>FilterTagsDialog.cs</DependentUpon>
-    </Compile>
-    <Compile Include="FindDialog.cs">
-      <SubType>Form</SubType>
-    </Compile>
-    <Compile Include="FindDialog.designer.cs">
-      <DependentUpon>FindDialog.cs</DependentUpon>
-    </Compile>
-    <Compile Include="HistogramParser.cs" />
-    <Compile Include="IOSSymbolParser.cs" />
-    <Compile Include="MainWindow.cs">
-      <SubType>Form</SubType>
-    </Compile>
-    <Compile Include="MainWindow.Designer.cs">
-      <DependentUpon>MainWindow.cs</DependentUpon>
-    </Compile>
-    <Compile Include="MemoryBitmapParser.cs" />
-    <Compile Include="MemoryPoolInfo.cs" />
-    <Compile Include="ModuleInfo.cs" />
-    <Compile Include="Properties\MetaData.cs" />
-    <Compile Include="OptionsDialog.cs">
-      <SubType>Form</SubType>
-    </Compile>
-    <Compile Include="OptionsDialog.Designer.cs">
-      <DependentUpon>OptionsDialog.cs</DependentUpon>
-    </Compile>
-    <Compile Include="ProfileDataHeader.cs" />
-    <Compile Include="ProfilingOptions.cs">
-      <SubType>Form</SubType>
-    </Compile>
-    <Compile Include="ProfilingOptions.designer.cs">
-      <DependentUpon>ProfilingOptions.cs</DependentUpon>
-    </Compile>
-    <Compile Include="Program.cs" />
-    <Compile Include="Properties\AssemblyInfo.cs" />
-    <EmbeddedResource Include="FilterTagsDialog.resx">
-      <DependentUpon>FilterTagsDialog.cs</DependentUpon>
-    </EmbeddedResource>
-    <EmbeddedResource Include="FindDialog.resx">
-      <DependentUpon>FindDialog.cs</DependentUpon>
-    </EmbeddedResource>
-    <EmbeddedResource Include="MainWindow.resx">
-      <SubType>Designer</SubType>
-      <DependentUpon>MainWindow.cs</DependentUpon>
-    </EmbeddedResource>
-    <EmbeddedResource Include="OptionsDialog.resx">
-      <DependentUpon>OptionsDialog.cs</DependentUpon>
-    </EmbeddedResource>
-    <EmbeddedResource Include="ProfilingOptions.resx">
-      <DependentUpon>ProfilingOptions.cs</DependentUpon>
-    </EmbeddedResource>
-    <EmbeddedResource Include="Properties\Resources.resx">
-      <Generator>ResXFileCodeGenerator</Generator>
-      <LastGenOutput>Resources.Designer.cs</LastGenOutput>
-      <SubType>Designer</SubType>
-    </EmbeddedResource>
-    <EmbeddedResource Include="SlowProgressDialog.resx">
-      <DependentUpon>SlowProgressDialog.cs</DependentUpon>
-    </EmbeddedResource>
-    <Compile Include="Properties\Resources.Designer.cs">
-      <AutoGen>True</AutoGen>
-      <DependentUpon>Resources.resx</DependentUpon>
-      <DesignTime>True</DesignTime>
-    </Compile>
-    <None Include="App.config" />
-    <None Include="Properties\Settings.settings">
-      <Generator>SettingsSingleFileGenerator</Generator>
-      <LastGenOutput>Settings.Designer.cs</LastGenOutput>
-    </None>
-    <Compile Include="Properties\Settings.Designer.cs">
-      <AutoGen>True</AutoGen>
-      <DependentUpon>Settings.settings</DependentUpon>
-      <DesignTimeSharedInput>True</DesignTimeSharedInput>
-    </Compile>
-  </ItemGroup>
-  <ItemGroup>
-    <Compile Include="ListViewEx.cs">
-      <SubType>Component</SubType>
-    </Compile>
-    <Compile Include="TreeViewEx.cs">
-      <SubType>Component</SubType>
-    </Compile>
-    <Compile Include="ScriptCallstack.cs" />
-    <Compile Include="ShortLivedAllocationView.cs" />
-    <Compile Include="SlowProgressDialog.cs">
-      <SubType>Form</SubType>
-    </Compile>
-    <Compile Include="SlowProgressDialog.designer.cs">
-      <DependentUpon>SlowProgressDialog.cs</DependentUpon>
-    </Compile>
-    <Compile Include="StreamInfo.cs" />
-    <Compile Include="StreamObserver.cs" />
-    <Compile Include="StreamParser.cs" />
-    <Compile Include="StreamSnapshot.cs" />
-    <Compile Include="StreamToken.cs" />
-    <Compile Include="TimeLineView.cs" />
-    <Compile Include="../../../Platforms/*/Source/Programs/MemoryProfiler2/**/*.cs">
-      <!-- Since RecursiveDir is bugged on VS Mac, we use the FullPath, and use Regex to pull out portions of the path to get Platform name -->
-      <Link>Platform/$([System.Text.RegularExpressions.Regex]::Replace(%(FullPath), '^(.*[\\/]Platforms[\\/])(.+?)[\\/](.*)$', '$2'))/%(Filename)%(Extension)</Link>
-    </Compile>
-  </ItemGroup>
-  <ItemGroup>
-    <BootstrapperPackage Include="Microsoft.Net.Client.3.5">
-      <Visible>False</Visible>
-      <ProductName>.NET Framework 3.5 SP1 Client Profile</ProductName>
-      <Install>false</Install>
-    </BootstrapperPackage>
-    <BootstrapperPackage Include="Microsoft.Net.Framework.2.0">
-      <Visible>False</Visible>
-      <ProductName>.NET Framework 2.0 %28x86%29</ProductName>
-      <Install>true</Install>
-    </BootstrapperPackage>
-    <BootstrapperPackage Include="Microsoft.Net.Framework.3.0">
-      <Visible>False</Visible>
-      <ProductName>.NET Framework 3.0 %28x86%29</ProductName>
-      <Install>false</Install>
-    </BootstrapperPackage>
-    <BootstrapperPackage Include="Microsoft.Net.Framework.3.5">
-      <Visible>False</Visible>
-      <ProductName>.NET Framework 3.5</ProductName>
-      <Install>false</Install>
-    </BootstrapperPackage>
-    <BootstrapperPackage Include="Microsoft.Net.Framework.3.5.SP1">
-      <Visible>False</Visible>
-      <ProductName>.NET Framework 3.5 SP1</ProductName>
-      <Install>false</Install>
-    </BootstrapperPackage>
-  </ItemGroup>
-  <ItemGroup />
-  <ItemGroup>
-    <None Include="Icons\Play.png" />
-  </ItemGroup>
-  <ItemGroup>
-    <None Include="Icons\PlayRed.png" />
-  </ItemGroup>
-  <ItemGroup>
-    <None Include="Icons\ExclamationMark.png" />
-  </ItemGroup>
-  <ItemGroup>
-    <Content Include="Icon.ico" />
-    <Content Include="Resources\Icon.ico" />
-  </ItemGroup>
-  <Import Project="$(MSBuildBinPath)\Microsoft.CSharp.targets" />
-  <!-- To modify your build process, add your task inside one of the targets below and uncomment it.
-       Other similar extension points exist, see Microsoft.Common.targets.
-  <Target Name="BeforeBuild">
-  </Target>
-  <Target Name="AfterBuild">
-  </Target>
-  -->
-</Project>
+
+</Project>
\ No newline at end of file
diff --git a/Engine/Source/Programs/MemoryProfiler2/Properties/AssemblyInfo.cs b/Engine/Source/Programs/MemoryProfiler2/Properties/AssemblyInfo.cs
--- a/Engine/Source/Programs/MemoryProfiler2/Properties/AssemblyInfo.cs
+++ b/Engine/Source/Programs/MemoryProfiler2/Properties/AssemblyInfo.cs
@@ -1,7 +1,6 @@
-﻿// Copyright Epic Games, Inc. All Rights Reserved.
+// Copyright Epic Games, Inc. All Rights Reserved.
 
 using System.Reflection;
 
 [assembly: AssemblyTitle( "Memory Profiler v2.1" )]
 [assembly: AssemblyDescription( "A tool to process the output from the in engine memory profiling tools." )]
-[assembly: AssemblyConfiguration("")]
diff --git a/Engine/Source/Programs/MemoryProfiler2/Properties/MetaData.cs b/Engine/Source/Programs/MemoryProfiler2/Properties/MetaData.cs
--- a/Engine/Source/Programs/MemoryProfiler2/Properties/MetaData.cs
+++ b/Engine/Source/Programs/MemoryProfiler2/Properties/MetaData.cs
@@ -15,9 +15,9 @@ using System.Resources;
 
 [assembly: ComVisible( false )]
 [assembly: NeutralResourcesLanguageAttribute( "" )]
-
+/*
 #if !SPECIFIC_VERSION
 // Automatically generate a version number based on the time of compilation
 [assembly: AssemblyVersion( "1.0.*" )]
 #endif
-
+*/
\ No newline at end of file
