Fix error when generating projects: "Existing C# project MemoryProfiler2.csproj must contain a Debug configuration".
This change works around the problem by adding empty PropertyGroups of the expected format.
#rnx

diff --git a/Engine/Source/Programs/MemoryProfiler2/MemoryProfiler2.csproj b/Engine/Source/Programs/MemoryProfiler2/MemoryProfiler2.csproj
--- a/Engine/Source/Programs/MemoryProfiler2/MemoryProfiler2.csproj
+++ b/Engine/Source/Programs/MemoryProfiler2/MemoryProfiler2.csproj
@@ -1,4 +1,4 @@
-﻿<Project Sdk="Microsoft.NET.Sdk">
+﻿<Project Sdk="Microsoft.NET.Sdk" xmlns="http://schemas.microsoft.com/developer/msbuild/2003">
 
   <PropertyGroup>
     <OutputType>WinExe</OutputType>
@@ -11,6 +11,9 @@
     <AssemblyName></AssemblyName>
   </PropertyGroup>
 
+  <PropertyGroup Condition=" '$(Configuration)|$(Platform)' == 'Debug|AnyCPU' "/>
+  <PropertyGroup Condition=" '$(Configuration)|$(Platform)' == 'Development|AnyCPU' "/>
+
   <ItemGroup>
     <PackageReference Include="envdte" Version="17.3.32804.24" />
     <PackageReference Include="HIC.System.Windows.Forms.DataVisualization" Version="1.0.1" />
