[MemoryProfiler2]
Force added configuration debug even though it is in there because UBT for some reason doesnt like the non-explicit set of configs

diff --git a/Engine/Source/Programs/MemoryProfiler2/MemoryProfiler2.csproj b/Engine/Source/Programs/MemoryProfiler2/MemoryProfiler2.csproj
--- a/Engine/Source/Programs/MemoryProfiler2/MemoryProfiler2.csproj
+++ b/Engine/Source/Programs/MemoryProfiler2/MemoryProfiler2.csproj
@@ -9,6 +9,7 @@
     <ProduceReferenceAssembly>True</ProduceReferenceAssembly>
     <PackageId />
     <AssemblyName></AssemblyName>
+    <Configurations>Debug;Release;Development</Configurations>
   </PropertyGroup>
 
   <PropertyGroup Condition=" '$(Configuration)|$(Platform)' == 'Debug|AnyCPU' "/>
