MemoryProfiler2.csproj: revert empty PropertyGroups fix; "must contain a Debug configuration" error has been fixed in another way.

diff --git a/Engine/Source/Programs/MemoryProfiler2/MemoryProfiler2.csproj b/Engine/Source/Programs/MemoryProfiler2/MemoryProfiler2.csproj
--- a/Engine/Source/Programs/MemoryProfiler2/MemoryProfiler2.csproj
+++ b/Engine/Source/Programs/MemoryProfiler2/MemoryProfiler2.csproj
@@ -1,4 +1,4 @@
-﻿<Project Sdk="Microsoft.NET.Sdk" xmlns="http://schemas.microsoft.com/developer/msbuild/2003">
+﻿<Project Sdk="Microsoft.NET.Sdk">
 
   <PropertyGroup>
     <OutputType>WinExe</OutputType>
@@ -12,9 +12,6 @@
     <Configurations>Debug;Release;Development</Configurations>
   </PropertyGroup>
 
-  <PropertyGroup Condition=" '$(Configuration)|$(Platform)' == 'Debug|AnyCPU' "/>
-  <PropertyGroup Condition=" '$(Configuration)|$(Platform)' == 'Development|AnyCPU' "/>
-
   <ItemGroup>
     <PackageReference Include="envdte" Version="17.3.32804.24" />
     <PackageReference Include="HIC.System.Windows.Forms.DataVisualization" Version="1.0.1" />
